User Asset Tagging: v3

diff --git a/Engine/Source/Editor/UserAssetTagsEditor/Private/Settings/ProjectUserAssetTagSettings.cpp b/Engine/Source/Editor/UserAssetTagsEditor/Private/Settings/ProjectUserAssetTagSettings.cpp
new file mode 100644
--- /dev/null
+++ b/Engine/Source/Editor/UserAssetTagsEditor/Private/Settings/ProjectUserAssetTagSettings.cpp
@@ -0,0 +1,19 @@
+﻿// Copyright Epic Games, Inc. All Rights Reserved.
+
+#include "ProjectUserAssetTagSettings.h"
+
+TSet<FName> UProjectUserAssetTagSettings::GetProjectUserAssetTagsForClass(const UClass* Class) const
+{
+	FSoftClassPath ClassPath(Class);
+	if(FavoriteUserAssetTagsPerClass.Contains(ClassPath))
+	{
+		return FavoriteUserAssetTagsPerClass[ClassPath].FavoriteUserAssetTags;
+	}
+
+	return {};
+}
+
+FName UProjectUserAssetTagSettings::GetCategoryName() const
+{
+	return NAME_Editor;
+}
diff --git a/Engine/Source/Editor/UserAssetTagsEditor/Private/Settings/ProjectUserAssetTagSettings.h b/Engine/Source/Editor/UserAssetTagsEditor/Private/Settings/ProjectUserAssetTagSettings.h
new file mode 100644
--- /dev/null
+++ b/Engine/Source/Editor/UserAssetTagsEditor/Private/Settings/ProjectUserAssetTagSettings.h
@@ -0,0 +1,25 @@
+﻿// Copyright Epic Games, Inc. All Rights Reserved.
+
+#pragma once
+
+#include "CoreMinimal.h"
+#include "Config/LocalFavoriteUserAssetTagsConfig.h"
+#include "Engine/DeveloperSettings.h"
+#include "ProjectUserAssetTagSettings.generated.h"
+
+/**
+ * 
+ */
+UCLASS(Config=UserAssetTags, DefaultConfig, DisplayName="User Asset Tags", MinimalAPI)
+class UProjectUserAssetTagSettings : public UDeveloperSettings
+{
+	GENERATED_BODY()
+
+public:
+	UPROPERTY(EditAnywhere, Config, Category="UserAssetTags")
+	TMap<FSoftClassPath, FPerTypeFavoriteUserAssetTags> FavoriteUserAssetTagsPerClass;
+
+	TSet<FName> GetProjectUserAssetTagsForClass(const UClass* Class) const;
+
+	virtual FName GetCategoryName() const override;
+};
diff --git a/Engine/Source/Editor/UserAssetTagsEditor/Private/UserAssetTagProviders/UserAssetTagProvider_LocalAssetTypeFavorites.h b/Engine/Source/Editor/UserAssetTagsEditor/Private/UserAssetTagProviders/UserAssetTagProvider_LocalAssetTypeFavorites.h
new file mode 100644
--- /dev/null
+++ b/Engine/Source/Editor/UserAssetTagsEditor/Private/UserAssetTagProviders/UserAssetTagProvider_LocalAssetTypeFavorites.h
@@ -0,0 +1,31 @@
+﻿// Copyright Epic Games, Inc. All Rights Reserved.
+
+#pragma once
+
+#include "CoreMinimal.h"
+#include "UserAssetTagProvider.h"
+#include "UserAssetTagProvider_LocalAssetTypeFavorites.generated.h"
+
+class UUserAssetTagEditorContext;
+
+/**
+ * The local favorites provider lets you add your own selection of favorite tags per asset type
+ * for quick reuse without affecting project settings.
+ */
+UCLASS(DisplayName="Favorites by type")
+class USERASSETTAGSEDITOR_API UUserAssetTagProvider_LocalAssetTypeFavorites : public UUserAssetTagProvider
+{
+	GENERATED_BODY()
+
+public:
+	UUserAssetTagProvider_LocalAssetTypeFavorites();
+
+	virtual FText GetDisplayNameText(const UUserAssetTagEditorContext* Context) const override;
+	virtual TSet<FName> GetSuggestedUserAssetTags(const UUserAssetTagEditorContext* Context) const override;
+
+	/** Favorites are per asset type, so we disable this provider if we encounter multiple asset types. */
+	virtual FResultWithUserFeedback IsValid(const UUserAssetTagEditorContext* Context) const override;
+
+	virtual TSharedPtr<SWidget> AddAdditionalSuggestedWidgets(FName UserAssetTag, const UUserAssetTagEditorContext* Context) const override;
+	virtual TSharedPtr<SWidget> AddAdditionalOwnedTagWidgets(FName UserAssetTag, const UUserAssetTagEditorContext* Context) const override;
+};
diff --git a/Engine/Source/Editor/UserAssetTagsEditor/Private/UserAssetTagProviders/UserAssetTagProvider_Project.h b/Engine/Source/Editor/UserAssetTagsEditor/Private/UserAssetTagProviders/UserAssetTagProvider_Project.h
new file mode 100644
--- /dev/null
+++ b/Engine/Source/Editor/UserAssetTagsEditor/Private/UserAssetTagProviders/UserAssetTagProvider_Project.h
@@ -0,0 +1,25 @@
+﻿// Copyright Epic Games, Inc. All Rights Reserved.
+
+#pragma once
+
+#include "CoreMinimal.h"
+#include "UserAssetTagProvider.h"
+#include "UserAssetTagProvider_Project.generated.h"
+
+/** 
+ *  A provider with the commonly used tags for this project, per asset type.
+ *  You can assign tags per asset type in the project settings.
+ */
+UCLASS(DisplayName="Project Tags by type")
+class USERASSETTAGSEDITOR_API UUserAssetTagProvider_Project : public UUserAssetTagProvider
+{
+	GENERATED_BODY()
+
+	virtual FText GetDisplayNameText(const UUserAssetTagEditorContext* Context) const override;
+	virtual TSet<FName> GetSuggestedUserAssetTags(const UUserAssetTagEditorContext* Context) const override;
+	virtual FResultWithUserFeedback IsValid(const UUserAssetTagEditorContext* Context) const override;
+	
+	static void NavigateToProjectTags();
+protected:
+	virtual void AddToolbarMenuEntries(class UToolMenu* DynamicMenu, const UUserAssetTagEditorContext* Context) const override;
+};
diff --git a/Engine/Source/Editor/UserAssetTagsEditor/Private/Widgets/STaggedAssetBrowserCreateAssetWindow.cpp b/Engine/Source/Editor/UserAssetTagsEditor/Private/Widgets/STaggedAssetBrowserCreateAssetWindow.cpp
new file mode 100644
--- /dev/null
+++ b/Engine/Source/Editor/UserAssetTagsEditor/Private/Widgets/STaggedAssetBrowserCreateAssetWindow.cpp
@@ -0,0 +1,152 @@
+﻿// Copyright Epic Games, Inc. All Rights Reserved.
+
+#include "Widgets/STaggedAssetBrowserCreateAssetWindow.h"
+#include "Widgets/Layout/SSpacer.h"
+
+#define LOCTEXT_NAMESPACE "TaggedAssetBrowser"
+
+void STaggedAssetBrowserCreateAssetWindow::Construct(const FArguments& InArgs, const UTaggedAssetBrowserConfiguration& Configuration, UClass& InCreatedClass)
+{
+	CreatedClass = &InCreatedClass;
+	AdditionalFactorySettingsClass = InArgs._AdditionalFactorySettingsClass;
+	bAllowEmptyAssetCreation = InArgs._bAllowEmptyAssetCreation;
+
+	FArguments Args = InArgs;
+
+	TSharedRef<SWidget> CreateAssetControls = SNew(SBox)
+		.Padding(16.f, 16.f)
+		[
+			SNew(SHorizontalBox)
+			+ SHorizontalBox::Slot()
+			.AutoWidth()
+			.Padding(5.f, 1.f)
+			.HAlign(HAlign_Left)
+			[
+				SNew(SButton)
+				.OnClicked(this, &STaggedAssetBrowserCreateAssetWindow::Proceed)
+				.Visibility(bAllowEmptyAssetCreation ? EVisibility::Visible : EVisibility::Hidden)
+				[
+					SNew(SHorizontalBox)
+					+ SHorizontalBox::Slot()
+					.AutoWidth()
+					.Padding(2.f)
+					[
+						SNew(SImage)
+						.Image(FSlateIconFinder::FindIconForClass(CreatedClass.Get()).GetIcon())
+					]
+					+ SHorizontalBox::Slot()
+					.AutoWidth()
+					.Padding(2.f)
+					[
+						SNew(STextBlock).Text(FText::FormatOrdered(LOCTEXT("CreateEmptyAssetButtonLabel", "Create Empty {0}"), CreatedClass->GetDisplayNameText()))
+					]
+				]
+			]
+			+ SHorizontalBox::Slot()
+			[
+				SNew(SSpacer)
+			]
+			+ SHorizontalBox::Slot()
+			.AutoWidth()
+			.Padding(5.f, 1.f)
+			[
+				SNew(SButton)
+				.ContentPadding(FMargin(50.f, 4.f, 50.f, 1.f))
+				.ButtonStyle(&FAppStyle::Get().GetWidgetStyle<FButtonStyle>("PrimaryButton"))
+				.TextStyle(&FAppStyle::Get().GetWidgetStyle<FTextBlockStyle>("PrimaryButtonText"))
+				.HAlign(HAlign_Center)
+				.Text(LOCTEXT("CreatePrimaryButtonLabel", "Create"))
+				.OnClicked(this, &STaggedAssetBrowserCreateAssetWindow::Proceed)
+				.IsEnabled(this, &STaggedAssetBrowserCreateAssetWindow::HasSelectedAssets)
+				.ToolTipText(this, &STaggedAssetBrowserCreateAssetWindow::GetCreateButtonTooltip)
+			]
+			+ SHorizontalBox::Slot()
+			.AutoWidth()
+			.Padding(5.f, 1.f)
+			[
+				SNew(SButton)
+				.ContentPadding(FMargin(50.f, 4.f, 50.f, 1.f))
+				.TextStyle(&FAppStyle::Get().GetWidgetStyle<FTextBlockStyle>("PrimaryButtonText"))
+				.HAlign(HAlign_Center)
+				.Text(LOCTEXT("CancelButtonLabel", "Cancel"))
+				.OnClicked(this, &STaggedAssetBrowserCreateAssetWindow::Cancel)
+			]
+		];
+
+	if(AdditionalFactorySettingsClass.IsValid())
+	{
+		FTaggedAssetBrowserCustomTabInfo FactoryTabInfo;
+		FactoryTabInfo.Title = LOCTEXT("FactoryTabTitle", "Settings");
+		FactoryTabInfo.Icon = FAppStyle::GetBrush("Icons.Settings");
+		FactoryTabInfo.OnGetTabContent = FOnGetContent::CreateSP(this, &STaggedAssetBrowserCreateAssetWindow::CreateFactorySettingsTab);
+
+		Args._AssetBrowserWindowArgs._AssetBrowserArgs._CustomTabInfos = { FactoryTabInfo };
+	}
+	
+	Args._AssetBrowserWindowArgs._AssetBrowserArgs._AdditionalBottomWidget = CreateAssetControls;
+	Args._AssetBrowserWindowArgs._AssetBrowserArgs._OnAssetsActivated = FOnAssetsActivated::CreateSP(this, &STaggedAssetBrowserCreateAssetWindow::OnAssetsActivatedInternal);
+	STaggedAssetBrowserWindow::Construct(Args._AssetBrowserWindowArgs, Configuration);
+}
+
+void STaggedAssetBrowserCreateAssetWindow::OnAssetsActivatedInternal(const TArray<FAssetData>& AssetData, EAssetTypeActivationMethod::Type)
+{
+	bProceedWithAction = true;
+}
+
+FReply STaggedAssetBrowserCreateAssetWindow::Proceed()
+{
+	// The assets don't matter here as the factories making use of this window will retrieve the selected assets afterwards
+	OnAssetsActivated({}, EAssetTypeActivationMethod::Opened);
+	return FReply::Handled();
+}
+
+FReply STaggedAssetBrowserCreateAssetWindow::Cancel()
+{
+	bProceedWithAction = false;
+	RequestDestroyWindow();
+	return FReply::Handled();
+}
+
+FText STaggedAssetBrowserCreateAssetWindow::GetCreateButtonTooltip() const
+{
+	return HasSelectedAssets()
+	? FText::FormatOrdered(LOCTEXT("CreateAssetButtonTooltip_Enabled", "Create a new {0} with selected asset {1}"), CreatedClass->GetDisplayNameText(), FText::FromName(GetSelectedAssets()[0].AssetName))
+	: FText::FormatOrdered(LOCTEXT("CreateAssetButtonTooltip_Disabled", "Please select an asset as a base for your new effect.{0}"), bAllowEmptyAssetCreation ? LOCTEXT("CreateAssetButtonTooltip_Disabled_AllowsEmptyCreation", " Alternatively, create an empty asset.") : FText::GetEmpty());
+}
+
+TSharedRef<SWidget> STaggedAssetBrowserCreateAssetWindow::CreateFactorySettingsTab()
+{
+	if(AdditionalFactorySettingsClass.IsValid())
+	{
+		FactorySettingsObject = NewObject<UObject>(GetTransientPackage(), AdditionalFactorySettingsClass.Get(), NAME_None, RF_Transient);
+		FPropertyEditorModule& PropertyEditorModule = FModuleManager::Get().LoadModuleChecked<FPropertyEditorModule>("PropertyEditor");
+
+		FDetailsViewArgs Args;
+		Args.bShowObjectLabel = false;
+		Args.NameAreaSettings = FDetailsViewArgs::HideNameArea;
+
+		FactorySettingsWidget = PropertyEditorModule.CreateDetailView(Args);
+		FactorySettingsWidget->SetObject(FactorySettingsObject);
+		
+		return FactorySettingsWidget.ToSharedRef();
+	}
+
+	return SNullWidget::NullWidget;
+}
+
+FString STaggedAssetBrowserCreateAssetWindow::GetReferencerName() const
+{
+	return TEXT("STaggedAssetBrowserCreateAsset");
+}
+
+void STaggedAssetBrowserCreateAssetWindow::AddReferencedObjects(FReferenceCollector& Collector)
+{
+	FSlateInvalidationRoot::AddReferencedObjects(Collector);
+	
+	if(FactorySettingsObject)
+	{
+		Collector.AddReferencedObject(FactorySettingsObject);
+	}
+}
+
+#undef LOCTEXT_NAMESPACE
diff --git a/Engine/Source/Editor/UserAssetTagsEditor/Public/AssetFactories/TaggedAssetBrowserConfigurationFactory.h b/Engine/Source/Editor/UserAssetTagsEditor/Public/AssetFactories/TaggedAssetBrowserConfigurationFactory.h
new file mode 100644
--- /dev/null
+++ b/Engine/Source/Editor/UserAssetTagsEditor/Public/AssetFactories/TaggedAssetBrowserConfigurationFactory.h
@@ -0,0 +1,17 @@
+﻿// Copyright Epic Games, Inc. All Rights Reserved.
+
+#pragma once
+
+#include "CoreMinimal.h"
+#include "Factories/Factory.h"
+#include "TaggedAssetBrowserConfigurationFactory.generated.h"
+
+UCLASS()
+class USERASSETTAGSEDITOR_API UTaggedAssetBrowserConfigurationFactory : public UFactory
+{
+	GENERATED_BODY()
+
+	UTaggedAssetBrowserConfigurationFactory();
+	
+	virtual UObject* FactoryCreateNew(UClass* InClass, UObject* InParent, FName InName, EObjectFlags Flags, UObject* Context, FFeedbackContext* Warn) override;
+};
diff --git a/Engine/Source/Editor/UserAssetTagsEditor/Public/Config/LocalFavoriteUserAssetTagsConfig.h b/Engine/Source/Editor/UserAssetTagsEditor/Public/Config/LocalFavoriteUserAssetTagsConfig.h
new file mode 100644
--- /dev/null
+++ b/Engine/Source/Editor/UserAssetTagsEditor/Public/Config/LocalFavoriteUserAssetTagsConfig.h
@@ -0,0 +1,43 @@
+﻿// Copyright Epic Games, Inc. All Rights Reserved.
+
+#pragma once
+
+#include "CoreMinimal.h"
+#include "EditorConfigBase.h"
+#include "LocalFavoriteUserAssetTagsConfig.generated.h"
+
+USTRUCT()
+struct FPerTypeFavoriteUserAssetTags
+{
+	GENERATED_BODY()
+
+	UPROPERTY(EditAnywhere, Category="User Asset Tags")
+	TSet<FName> FavoriteUserAssetTags;
+};
+/**
+ * 
+ */
+UCLASS(EditorConfig="FavoriteUserAssetTags")
+class USERASSETTAGSEDITOR_API ULocalFavoriteUserAssetTagsConfig : public UEditorConfigBase
+{
+	GENERATED_BODY()
+
+public:
+	static ULocalFavoriteUserAssetTagsConfig* Get();
+	static void Shutdown();
+	
+	TMap<FSoftClassPath, FPerTypeFavoriteUserAssetTags>& GetFavoriteUserAssetTagsMutable() { return FavoriteUserAssetTagsPerClass; }
+	
+	const TMap<FSoftClassPath, FPerTypeFavoriteUserAssetTags>& GetFavoriteUserAssetTags() const { return FavoriteUserAssetTagsPerClass; }
+
+	TSet<FName> GetFavoriteUserAssetTagsForClass(const UClass* Class);
+	void ToggleFavoriteUserAssetTag(const UClass* Class, FName InUserAssetTag);
+private:
+	UPROPERTY(meta=(EditorConfig))
+	TMap<FSoftClassPath, FPerTypeFavoriteUserAssetTags> FavoriteUserAssetTagsPerClass;
+
+	UPROPERTY(meta=(EditorConfig))
+	int32 MaxRecentUserAssetTags = 10;
+		
+	static TStrongObjectPtr<ULocalFavoriteUserAssetTagsConfig> Instance;
+};
diff --git a/Engine/Source/Editor/UserAssetTagsEditor/Public/Config/TaggedAssetBrowserConfig.h b/Engine/Source/Editor/UserAssetTagsEditor/Public/Config/TaggedAssetBrowserConfig.h
new file mode 100644
--- /dev/null
+++ b/Engine/Source/Editor/UserAssetTagsEditor/Public/Config/TaggedAssetBrowserConfig.h
@@ -0,0 +1,44 @@
+﻿// Copyright Epic Games, Inc. All Rights Reserved.
+
+#pragma once
+
+#include "CoreMinimal.h"
+#include "EditorConfigBase.h"
+#include "TaggedAssetBrowserConfig.generated.h"
+
+USTRUCT()
+struct FPerTaggedAssetBrowserSavedState
+{
+	GENERATED_BODY()
+	
+	UPROPERTY()
+	TArray<FName> PrimaryFilterSelection;
+};
+
+UCLASS(EditorConfig="TaggedAssetBrowser")
+class USERASSETTAGSEDITOR_API UTaggedAssetBrowserConfig : public UEditorConfigBase
+{
+	GENERATED_BODY()
+public:
+	DECLARE_MULTICAST_DELEGATE_OneParam(FOnPropertyChanged, const FPropertyChangedEvent&);
+		
+	static UTaggedAssetBrowserConfig* Get();
+	static void Shutdown();
+	
+	FOnPropertyChanged& OnPropertyChanged() { return OnPropertyChangedDelegate; }
+	
+	UPROPERTY(meta=(EditorConfig))
+	TMap<FName, FPerTaggedAssetBrowserSavedState> PerTaggedAssetBrowserSettings;
+	
+	UPROPERTY(meta=(EditorConfig))
+    bool bShowHiddenAssets = false;
+    	
+	UPROPERTY(meta=(EditorConfig))
+	bool bShowDeprecatedAssets = false;
+
+	virtual void PostEditChangeProperty(FPropertyChangedEvent& PropertyChangedEvent) override;	
+private:
+	static TStrongObjectPtr<UTaggedAssetBrowserConfig> Instance;
+
+	FOnPropertyChanged OnPropertyChangedDelegate;
+};
diff --git a/Engine/Source/Editor/UserAssetTagsEditor/Public/Config/UserAssetTagsEditorConfig.h b/Engine/Source/Editor/UserAssetTagsEditor/Public/Config/UserAssetTagsEditorConfig.h
new file mode 100644
--- /dev/null
+++ b/Engine/Source/Editor/UserAssetTagsEditor/Public/Config/UserAssetTagsEditorConfig.h
@@ -0,0 +1,60 @@
+﻿// Copyright Epic Games, Inc. All Rights Reserved.
+
+#pragma once
+
+#include "CoreMinimal.h"
+#include "EditorConfigBase.h"
+#include "UserAssetTagProvider.h"
+#include "UObject/StrongObjectPtrTemplates.h"
+#include "UserAssetTagsEditorConfig.generated.h"
+
+USTRUCT()
+struct FPerUserAssetTagProviderViewOptions
+{
+	GENERATED_BODY()
+
+	UPROPERTY()
+	bool bEnabled = true;
+
+	/** 0: Section; 1: SubMenu. */
+	UPROPERTY()
+	EUserAssetTagProviderMenuType MenuType = EUserAssetTagProviderMenuType::Section;
+};
+
+USTRUCT()
+struct FUserAssetTagProviderViewOptions
+{
+	GENERATED_BODY()
+	
+	UPROPERTY()
+	TMap<FName, FPerUserAssetTagProviderViewOptions> PerProviderViewOptions;
+};
+/**
+ * 
+ */
+UCLASS(EditorConfig=UserAssetTags)
+class USERASSETTAGSEDITOR_API UUserAssetTagsEditorConfig : public UEditorConfigBase
+{
+	GENERATED_BODY()
+
+public:
+	void ToggleSortByAlphabet();
+	bool ShouldSortByAlphabet() const;
+	
+	bool IsProviderEnabled(const UClass* ProviderClass) const;
+	void ToggleProviderEnabled(const UClass* ProviderClass);
+	
+	EUserAssetTagProviderMenuType GetProviderMenuType(const UClass* ProviderClass) const;
+	void SetProviderMenuType(const UClass* ProviderClass, EUserAssetTagProviderMenuType InMenuType);
+
+	static UUserAssetTagsEditorConfig* Get();
+	static void Shutdown();
+private:
+	UPROPERTY(meta=(EditorConfig))
+	bool bSortByAlphabet = false;
+	
+	UPROPERTY(meta=(EditorConfig))
+	FUserAssetTagProviderViewOptions ProviderViewOptions;
+private:
+    static TStrongObjectPtr<UUserAssetTagsEditorConfig> Instance;
+};
diff --git a/Engine/Source/Editor/UserAssetTagsEditor/Public/TaggedAssetBrowserFilters/TaggedAssetBrowser_CommonFilters.h b/Engine/Source/Editor/UserAssetTagsEditor/Public/TaggedAssetBrowserFilters/TaggedAssetBrowser_CommonFilters.h
new file mode 100644
--- /dev/null
+++ b/Engine/Source/Editor/UserAssetTagsEditor/Public/TaggedAssetBrowserFilters/TaggedAssetBrowser_CommonFilters.h
@@ -0,0 +1,151 @@
+﻿// Copyright Epic Games, Inc. All Rights Reserved.
+
+#pragma once
+
+#include "CoreMinimal.h"
+#include "TaggedAssetBrowserMenuFilters.h"
+#include "UObject/Object.h"
+#include "MRUFavoritesList.h"
+#include "TaggedAssetBrowser_CommonFilters.generated.h"
+
+class UTaggedAssetBrowserSection;
+
+UCLASS()
+class UTaggedAssetBrowserFilterRoot : public UHierarchyRoot
+{
+	GENERATED_BODY()
+};
+
+USTRUCT()
+struct FTaggedAssetBrowserSectionIconData
+{
+	GENERATED_BODY()
+
+	/** If true, will use a texture. */
+	UPROPERTY(EditAnywhere, Category="Icon")
+	bool bUseTextureForIcon = false;
+
+	/** The actual name of the icon to use within the styleset. */
+	UPROPERTY(EditAnywhere, Category="Icon", meta=(EditCondition="bUseTextureForIcon==false"))
+	FName StyleName;
+
+	/** The icon to display for this section. */
+	UPROPERTY(EditAnywhere, Category = "Icon", meta=(EditCondition="bUseTextureForIcon==true"))
+	TObjectPtr<UTexture2D> Icon;
+
+	const FSlateBrush* GetImageBrush() const;
+private:
+	mutable FSlateBrush TextureBrush;
+};
+
+/** The section class for the tagged asset browser. Can contain additional filters.*/
+UCLASS(MinimalAPI)
+class UTaggedAssetBrowserSection : public UHierarchySection
+{
+	GENERATED_BODY()
+
+public:
+	UTaggedAssetBrowserSection() = default;
+
+	UPROPERTY(EditAnywhere, Category = "Section", meta=(ShowOnlyInnerProperties))
+	FTaggedAssetBrowserSectionIconData IconData;
+
+	/** A list of filters that can be specified optionally. This way, a section can display its associated elements, OR additionally filter assets by itself. */
+	UPROPERTY(EditAnywhere, Category = "Section", Instanced)
+	TArray<TObjectPtr<UTaggedAssetBrowserFilterBase>> Filters;
+};
+
+/** Show all assets. */
+UCLASS(DisplayName="All")
+class USERASSETTAGSEDITOR_API UTaggedAssetBrowserFilter_All : public UTaggedAssetBrowserFilterBase
+{
+	GENERATED_BODY()
+public:
+	UTaggedAssetBrowserFilter_All();
+};
+
+/** A filter that returns assets that have the specified User Asset Tag. Nested User Asset Tag filters will be combined. */
+UCLASS(DisplayName="User Asset Tag")
+class USERASSETTAGSEDITOR_API UTaggedAssetBrowserFilter_UserAssetTag : public UTaggedAssetBrowserFilterBase
+{
+	GENERATED_BODY()
+public:
+	UTaggedAssetBrowserFilter_UserAssetTag() = default;
+	
+	virtual FString ToString() const override;
+	virtual TOptional<FString> GetInstanceIdentifier() const override { return ToString(); }
+	virtual FSlateIcon GetIcon() const override;
+	virtual void ModifyARFilterInternal(FARFilter& Filter) const override;
+	virtual bool ShouldFilterAssetInternal(const FAssetData& InAssetData) const override;
+
+	void SetUserAssetTag(FName InUserAssetTag);
+	bool DoesAssetHaveTag(const FAssetData& AssetCandidate) const;
+
+protected:
+	UPROPERTY(EditAnywhere, Category="Filter")
+	FName UserAssetTag;
+};
+
+/** A collection of multiple User Asset Tags. Will show all assets that own at least one of the contained tags. */
+UCLASS(DisplayName="User Asset Tag Collection")
+class USERASSETTAGSEDITOR_API UTaggedAssetBrowserFilter_UserAssetTagCollection : public UTaggedAssetBrowserFilterBase
+{
+	GENERATED_BODY()
+public:
+	UTaggedAssetBrowserFilter_UserAssetTagCollection();
+
+	virtual FString ToString() const override;
+	virtual TOptional<FString> GetInstanceIdentifier() const override { return ToString(); }
+	virtual FText GetTooltip() const override;
+	virtual FSlateIcon GetIcon() const override;
+	virtual void CreateAdditionalWidgets(TSharedPtr<SHorizontalBox> ExtensionBox) override;
+
+	virtual void ModifyARFilterInternal(FARFilter& Filter) const override;
+	virtual bool ShouldFilterAssetInternal(const FAssetData& InAssetData) const override;
+	
+	void SetCollectionName(FName InName);
+	
+protected:
+	UPROPERTY(EditAnywhere, Category="Filter")
+	FName Name;
+	
+	UPROPERTY(EditAnywhere, Category="Filter", meta=(MultiLine))
+	FText Description;
+	
+private:
+	FText GetContainedChildrenNumberText() const;
+	EVisibility GetContainedChildrenNumberTextVisibility() const;
+	FText GetContainedChildrenNumberTooltip() const;
+};
+
+/** Filter for recently used assets. */
+UCLASS(DisplayName="Recent")
+class USERASSETTAGSEDITOR_API UTaggedAssetBrowserFilter_Recent : public UTaggedAssetBrowserFilterBase
+{
+	GENERATED_BODY()
+public:
+	UTaggedAssetBrowserFilter_Recent() = default;
+	
+	virtual void InitializeInternal(const FTaggedAssetBrowserContext& InContext) override;
+	virtual bool ShouldFilterAssetInternal(const FAssetData& InAssetData) const override;
+	virtual FSlateIcon GetIcon() const override;
+
+	bool IsAssetRecent(const FAssetData& AssetCandidate) const;
+
+	TAttribute<const FMainMRUFavoritesList*> ListAttribute;
+};
+
+/** Filter for assets based on directories. */
+UCLASS(DisplayName="Directories")
+class USERASSETTAGSEDITOR_API UTaggedAssetBrowserFilter_Directories : public UTaggedAssetBrowserFilterBase
+{
+	GENERATED_BODY()
+public:
+	UTaggedAssetBrowserFilter_Directories() = default;
+	
+	virtual void ModifyARFilterInternal(FARFilter& Filter) const override;
+	virtual FSlateIcon GetIcon() const override;
+
+	UPROPERTY(EditAnywhere, Category="Filter", meta=(ContentDir))
+	TArray<FDirectoryPath> DirectoryPaths;
+};
diff --git a/Engine/Source/Editor/UserAssetTagsEditor/Public/UserAssetTagsEditorModule.h b/Engine/Source/Editor/UserAssetTagsEditor/Public/UserAssetTagsEditorModule.h
new file mode 100644
--- /dev/null
+++ b/Engine/Source/Editor/UserAssetTagsEditor/Public/UserAssetTagsEditorModule.h
@@ -0,0 +1,53 @@
+﻿// Copyright Epic Games, Inc. All Rights Reserved.
+
+#pragma once
+
+#include "DataHierarchyViewModelBase.h"
+#include "Modules/ModuleManager.h"
+#include "Delegates/DelegateCombinations.h"
+
+struct FHierarchyElementViewModel;
+class UTaggedAssetBrowserFilterBase;
+class SDockTab;
+class FSpawnTabArgs;
+class FAssetRegistryTagsContext;
+
+USERASSETTAGSEDITOR_API DECLARE_LOG_CATEGORY_EXTERN(LogUserAssetTags, Log, All);
+
+class USERASSETTAGSEDITOR_API FUserAssetTagsEditorModule : public IModuleInterface
+{
+public:
+	DECLARE_DELEGATE_RetVal_ThreeParams(TSharedPtr<FHierarchyElementViewModel>, FOnGetViewModelFactory, UHierarchyElement* Element, TSharedPtr<FHierarchyElementViewModel> InParent, TWeakObjectPtr<UDataHierarchyViewModelBase> InViewModel);
+
+	FUserAssetTagsEditorModule();
+	virtual ~FUserAssetTagsEditorModule() override;
+	
+	static FUserAssetTagsEditorModule& Get();
+	
+	virtual void StartupModule() override;
+	virtual void ShutdownModule() override;
+
+	/** To register a custom view model for display in the Tagged Asset Browser. */
+	void RegisterConfigurationHierarchyElementViewModel(TSubclassOf<UHierarchyElement> FilterType, FOnGetViewModelFactory InFactory);
+	void UnregisterConfigurationHierarchyElementViewModel(TSubclassOf<UHierarchyElement> FilterType);
+	const TMap<TSubclassOf<UHierarchyElement>, FOnGetViewModelFactory>& GetConfigurationHierarchyViewModelFactories() const { return ConfigurationHierarchyViewModelFactories; }
+
+	static FName ManageTagsTabId;
+private:
+	void RegisterDefaultConfigurationHierarchyElementViewModels();
+	void RegisterDetailsCustomizations();
+
+	/** We extend the Content Browser command list so that the shortcut works. */
+	void OnRegisterCommandList(FName ContextName, TSharedRef<FUICommandList> CommandList);
+
+	static TSharedRef<SDockTab> SpawnUserAssetTagsEditorNomadTab(const FSpawnTabArgs& SpawnTabArgs);
+
+	static void GetAssetRegistryTagsFromUserAssetTags(FAssetRegistryTagsContext AssetRegistryTagsContext);
+	static void TransferUserAssetTagsFromOldPackageToNewPackage(const FMetaData& OldMetaData, FMetaData& NewMetaData, const UPackage* OldPackage, UPackage* NewPackage);
+
+private:
+	TMap<TSubclassOf<UHierarchyElement>, FOnGetViewModelFactory> ConfigurationHierarchyViewModelFactories;
+	TSharedPtr<FUICommandList> CommandList;
+	FDelegateHandle GetAssetRegistryTagsFromUserAssetTagsDelegateHandle;
+	FDelegateHandle TransferUserAssetTagsFromOldPackageToNewPackageDelegateHandle;
+};
diff --git a/Engine/Source/Editor/UserAssetTagsEditor/Public/Widgets/STaggedAssetBrowserCreateAssetWindow.h b/Engine/Source/Editor/UserAssetTagsEditor/Public/Widgets/STaggedAssetBrowserCreateAssetWindow.h
new file mode 100644
--- /dev/null
+++ b/Engine/Source/Editor/UserAssetTagsEditor/Public/Widgets/STaggedAssetBrowserCreateAssetWindow.h
@@ -0,0 +1,54 @@
+﻿// Copyright Epic Games, Inc. All Rights Reserved.
+
+#pragma once
+
+#include "MRUFavoritesList.h"
+#include "STaggedAssetBrowser.h"
+#include "STaggedAssetBrowserContent.h"
+#include "Editor/EditorEngine.h"
+#include "Subsystems/AssetEditorSubsystem.h"
+#include "Widgets/SCompoundWidget.h"
+#include "Widgets/Views/STreeView.h"
+
+class USERASSETTAGSEDITOR_API STaggedAssetBrowserCreateAssetWindow : public STaggedAssetBrowserWindow
+{
+	SLATE_BEGIN_ARGS(STaggedAssetBrowserCreateAssetWindow)
+		: _AdditionalFactorySettingsClass(nullptr)
+		, _bAllowEmptyAssetCreation(true)
+		{
+		}
+		SLATE_ARGUMENT(STaggedAssetBrowserWindow::FArguments, AssetBrowserWindowArgs)
+		/** An additional factory settings class can be specified for additional factory information. */
+		SLATE_ARGUMENT(UClass*, AdditionalFactorySettingsClass)
+		SLATE_ARGUMENT(bool, bAllowEmptyAssetCreation)
+	SLATE_END_ARGS()
+
+	void Construct(const FArguments& InArgs, const UTaggedAssetBrowserConfiguration& Configuration, UClass& CreatedClass);
+
+	bool ShouldProceedWithAction() const { return bProceedWithAction; }
+
+	/** Retrieve the factory settings object, if available. This can be used to further customize the factory process. */
+	UObject* RetrieveFactorySettings() const { return FactorySettingsObject; }
+private:
+	/** The function that will be called by our buttons or by the asset picker itself if double-clicking, hitting enter etc. */
+	void OnAssetsActivatedInternal(const TArray<FAssetData>& AssetData, EAssetTypeActivationMethod::Type);
+	
+	FReply Proceed();
+	FReply Cancel();
+
+	FText GetCreateButtonTooltip() const;
+	
+	TSharedRef<SWidget> CreateFactorySettingsTab();
+
+	virtual FString GetReferencerName() const override;
+	virtual void AddReferencedObjects(FReferenceCollector& Collector) override;
+private:
+	TWeakObjectPtr<UClass> CreatedClass;
+	TWeakObjectPtr<UClass> AdditionalFactorySettingsClass;
+
+	TObjectPtr<UObject> FactorySettingsObject;
+	TSharedPtr<class IDetailsView> FactorySettingsWidget;
+	
+	bool bProceedWithAction = false;
+	bool bAllowEmptyAssetCreation = false;
+};
diff --git a/Engine/Source/Editor/UserAssetTagsEditor/Public/Widgets/SUserAssetTagsEditor.h b/Engine/Source/Editor/UserAssetTagsEditor/Public/Widgets/SUserAssetTagsEditor.h
new file mode 100644
--- /dev/null
+++ b/Engine/Source/Editor/UserAssetTagsEditor/Public/Widgets/SUserAssetTagsEditor.h
@@ -0,0 +1,106 @@
+﻿// Copyright Epic Games, Inc. All Rights Reserved.
+
+#pragma once
+
+#include "UserAssetTagProvider.h"
+#include "Widgets/SCompoundWidget.h"
+#include "AssetRegistry/AssetData.h"
+#include "Templates/SharedPointer.h"
+#include "UserAssetTagEditorMenuContexts.h"
+
+struct FToolMenuContext;
+struct FToolUIAction;
+struct FToolMenuCustomWidgetContext;
+
+/** Optional info relating to a specific tag, useful to surface more information to the user. */
+struct FUserAssetTagInfo
+{
+	/** A tag can come from multiple sources at once. */
+	TArray<const UUserAssetTagProvider*> Sources;
+};
+
+/** The actual editor widget to assign/unassign tags. */
+class USERASSETTAGSEDITOR_API SUserAssetTagsEditor : public SCompoundWidget
+{
+public:
+	SLATE_BEGIN_ARGS(SUserAssetTagsEditor)
+		{
+		}
+
+	SLATE_END_ARGS()
+
+	/** Constructs this widget with InArgs */
+	void Construct(const FArguments& InArgs);
+	virtual ~SUserAssetTagsEditor() override;
+	
+	const TArray<TSharedPtr<FAssetData>>& GetSelectedAssets() const { return SelectedAssets; }
+
+	void RefreshDataAndMenus();
+	static void RefreshMenus();
+	static void RefreshSelectedAssetsMenu();
+	static void RefreshSuggestedTagMenus();
+	static void RefreshOwnedTagMenus();
+	static void RefreshProviderExtensionToolbarMenu();
+
+private:
+	TSharedPtr<SWidget> CreateSelectedAssetList();
+	TSharedPtr<SWidget> CreateSuggestedTagList();
+	TSharedPtr<SWidget> CreateOwnedTagsList();
+	TSharedPtr<SWidget> CreateToolbar();
+
+	static TArray<FAssetData> GetCurrentContentBrowserAssetSelection();
+	
+	void OnCommitNewTag(const FText& Text, ETextCommit::Type CommitType);
+	FReply OnAddTagButtonClicked();
+	bool IsAddTagButtonEnabled() const;
+
+	const FSlateBrush* GetViewOptionsBadgeIcon() const;
+	
+	static TSharedRef<SWidget> GenerateRowContent_SelectedAsset(const FAssetData& AssetData, const UUserAssetTagEditorContext* Context);
+	static TSharedRef<SWidget> GenerateRowContent_SuggestedTag(FName UserAssetTag, const UUserAssetTagEditorContext* Context);
+	static TSharedRef<SWidget> GenerateRowContent_OwnedTag(FName UserAssetTag, const UUserAssetTagEditorContext* Context);
+	static TSharedRef<SWidget> CreateMenuControlWidget(const FToolMenuContext& ToolMenuContext, const FToolMenuCustomWidgetContext& ToolMenuCustomWidgetContext, const UClass* Class);
+
+	TSharedRef<SWidget> OnGetViewOptions();
+	static void ToggleSortByAlphabet(const FToolMenuContext& ToolMenuContext);
+	static ECheckBoxState GetShouldSortByAlphabet(const FToolMenuContext& ToolMenuContext);
+	static void ToggleViewOption_ProviderClass(const FToolMenuContext& ToolMenuContext, const UClass* ProviderClass);
+	static ECheckBoxState GetViewOptions_IsProviderClassEnabled(const FToolMenuContext& ToolMenuContext, const UClass* Class);
+
+	static void SetViewOption_ProviderClassMenuType(EUserAssetTagProviderMenuType InMenuType, const UClass* ProviderClass);
+	static EUserAssetTagProviderMenuType GetViewOptions_ProviderClassMenuType(const UClass* ProviderClass);
+	
+	static const TArray<const UUserAssetTagProvider*>& GetAllProviderCDOs();
+	static TArray<const UUserAssetTagProvider*> GetAllValidDefaultEnabledProviderCDOs(const UUserAssetTagEditorContext* Context);
+	static TArray<const UUserAssetTagProvider*> GetValidProviderCDOs(const UUserAssetTagEditorContext* Context);
+	static TArray<const UUserAssetTagProvider*> GetEnabledProviderCDOs(const UUserAssetTagEditorContext* Context);
+
+	static FToolUIAction CreateToggleTagCheckboxAction_SuggestedTag(const FName& InUserAssetTag);
+	static FToolUIAction CreateToggleTagCheckboxAction_OwnedTag(const FName& InUserAssetTag);
+
+	virtual FReply OnKeyDown(const FGeometry& MyGeometry, const FKeyEvent& InKeyEvent) override;
+
+	static FReply OnDocumentationRequested();
+	static FText GetDocumentationButtonTooltipText();
+
+	static FReply OnDeleteTagClicked(FName InUserAssetTag, const UUserAssetTagEditorContext* Context);
+
+	void HandleAssetSelectionChanged(const TArray<FAssetData>& AssetData, bool bIsPrimaryBrowser);
+	
+	static TArray<FAssetData> TransformAssetData(const TArray<TSharedPtr<FAssetData>>& InAssetData);
+private:
+	TSharedPtr<SWidget> SelectedAssetList;
+	TSharedPtr<SWidget> TagSuggestionList;
+	TSharedPtr<SWidget> OwnedTagsList;
+
+	TArray<TSharedPtr<FAssetData>> SelectedAssets;
+	
+	TSharedPtr<class SEditableTextBox> AddTagTextBox;
+	TSharedPtr<SWidget> ToolbarWidget;
+
+	/** The editor relies on UToolMenus a lot. We construct a context object once, and then pass it around to mostly static functions that generate the UI. */
+	TStrongObjectPtr<UUserAssetTagEditorContext> ThisContext;
+
+private:
+	static TArray<const UUserAssetTagProvider*> CachedProviderCDOs;
+};
diff --git a/Engine/Source/Runtime/CoreUObject/Private/UObject/Package.cpp b/Engine/Source/Runtime/CoreUObject/Private/UObject/Package.cpp
--- a/Engine/Source/Runtime/CoreUObject/Private/UObject/Package.cpp
+++ b/Engine/Source/Runtime/CoreUObject/Private/UObject/Package.cpp
@@ -87,6 +87,10 @@ UPackage::FOnPackageDirtyStateChanged UPackage::PackageDirtyStateChangedEvent;
  */
 UPackage::FOnPackageMarkedDirty UPackage::PackageMarkedDirtyEvent;
 
+#if WITH_METADATA
+UPackage::FOnMetaDataTransferRequestedDelegate UPackage::OnMetaDataTransferRequestedDelegate;
+#endif
+
 PRAGMA_DISABLE_DEPRECATION_WARNINGS // Silence deprecation warnings for deprecated CookedHash and LinkerSave members
 FSavePackageResultStruct::FSavePackageResultStruct()
 	: Result(ESavePackageResult::Error), TotalFileSize(0), SerializedPackageFlags(0)
diff --git a/Engine/Source/Runtime/CoreUObject/Public/UObject/Package.h b/Engine/Source/Runtime/CoreUObject/Public/UObject/Package.h
--- a/Engine/Source/Runtime/CoreUObject/Public/UObject/Package.h
+++ b/Engine/Source/Runtime/CoreUObject/Public/UObject/Package.h
@@ -244,6 +244,13 @@ public:
 	/** delegate type for when a package is about to be saved */
 	DECLARE_MULTICAST_DELEGATE_TwoParams(FPreSavePackageWithContext, class UPackage*, FObjectPreSaveContext);
 
+#if WITH_METADATA
+	/** A delegate that is called when old metadata is supposed to move to new metadata, such as when renaming an asset.
+	 *  Since metadata is not handled by the UPROPERTY-system, data transfer has to be done manually.
+	 *  By default, data is thrown out.  */
+	DECLARE_MULTICAST_DELEGATE_FourParams(FOnMetaDataTransferRequestedDelegate, const FMetaData& OldMetaData, FMetaData& NewMetaData, const UPackage* OldPackage, UPackage* NewPackage);
+#endif
+	
 	UE_DEPRECATED(5.0, "Use PreSavePackageWithContextEvent instead.")
 	COREUOBJECT_API static FPreSavePackage PreSavePackageEvent;
 	/** Delegate to notify subscribers when a package is about to be saved. */
@@ -265,6 +272,9 @@ public:
 	*/
 	COREUOBJECT_API static FOnPackageMarkedDirty PackageMarkedDirtyEvent;
 
+#if WITH_METADATA
+	COREUOBJECT_API static FOnMetaDataTransferRequestedDelegate OnMetaDataTransferRequestedDelegate;
+#endif
 private:
 	/** Used by the editor to determine if a package has been changed. */
 	uint8	bDirty:1;
diff --git a/Engine/Source/Runtime/SlateCore/Private/Styling/StarshipCoreStyle.cpp b/Engine/Source/Runtime/SlateCore/Private/Styling/StarshipCoreStyle.cpp
--- a/Engine/Source/Runtime/SlateCore/Private/Styling/StarshipCoreStyle.cpp
+++ b/Engine/Source/Runtime/SlateCore/Private/Styling/StarshipCoreStyle.cpp
@@ -355,6 +355,7 @@ TSharedRef<ISlateStyle> FStarshipCoreStyle::Create()
 		Style->Set("Icons.AutoFilter", new IMAGE_BRUSH_SVG("Starship/Common/FilterAuto", Icon16x16));
 
 		Style->Set("Icons.Tag", new IMAGE_BRUSH_SVG("Starship/Common/TagSimple_16", Icon16x16));
+		Style->Set("Icons.TagCollection", new IMAGE_BRUSH_SVG("Starship/Common/TagCollection_16", Icon16x16));
 		
 		Style->Set("Icons.Lock", new IMAGE_BRUSH_SVG("Starship/Common/lock", Icon16x16));
 		Style->Set("Icons.Unlock", new IMAGE_BRUSH_SVG("Starship/Common/lock-unlocked", Icon16x16));
