Added Class based filter to Tagged Asset Browser

diff --git a/Engine/Source/Editor/UserAssetTagsEditor/Private/Assets/TaggedAssetBrowserConfiguration.cpp b/Engine/Source/Editor/UserAssetTagsEditor/Private/Assets/TaggedAssetBrowserConfiguration.cpp
--- a/Engine/Source/Editor/UserAssetTagsEditor/Private/Assets/TaggedAssetBrowserConfiguration.cpp
+++ b/Engine/Source/Editor/UserAssetTagsEditor/Private/Assets/TaggedAssetBrowserConfiguration.cpp
@@ -18,12 +18,14 @@ UTaggedAssetBrowserConfiguration::UTaggedAssetBrowserConfiguration()
 		UTaggedAssetBrowserFilter_UserAssetTag::StaticClass(),
 		UTaggedAssetBrowserFilter_UserAssetTagCollection::StaticClass(),
 		UTaggedAssetBrowserFilter_Recent::StaticClass(),
-		UTaggedAssetBrowserFilter_Directories::StaticClass()
+		UTaggedAssetBrowserFilter_Directories::StaticClass(),
+		UTaggedAssetBrowserFilter_Class::StaticClass()
 	};
 	StandaloneData.ExtensionFilterClasses = {
 		UTaggedAssetBrowserFilter_UserAssetTag::StaticClass(),
 		UTaggedAssetBrowserFilter_UserAssetTagCollection::StaticClass(),
 		UTaggedAssetBrowserFilter_Directories::StaticClass(),
+		UTaggedAssetBrowserFilter_Class::StaticClass()
 	};
 }
 
diff --git a/Engine/Source/Editor/UserAssetTagsEditor/Private/TaggedAssetBrowserFilters/TaggedAssetBrowser_CommonFilters.cpp b/Engine/Source/Editor/UserAssetTagsEditor/Private/TaggedAssetBrowserFilters/TaggedAssetBrowser_CommonFilters.cpp
--- a/Engine/Source/Editor/UserAssetTagsEditor/Private/TaggedAssetBrowserFilters/TaggedAssetBrowser_CommonFilters.cpp
+++ b/Engine/Source/Editor/UserAssetTagsEditor/Private/TaggedAssetBrowserFilters/TaggedAssetBrowser_CommonFilters.cpp
@@ -252,4 +252,15 @@ FSlateIcon UTaggedAssetBrowserFilter_Directories::GetIcon() const
 	return FSlateIcon(FAppStyle::GetAppStyleSetName(), "Icons.FolderClosed");
 }
 
+void UTaggedAssetBrowserFilter_Class::ModifyARFilterInternal(FARFilter& Filter) const
+{
+	for(UClass* Class : Classes)
+	{
+		if(Class)
+		{
+			Filter.ClassPaths.Add(Class->GetClassPathName());
+		}
+	}
+}
+
 #undef LOCTEXT_NAMESPACE
diff --git a/Engine/Source/Editor/UserAssetTagsEditor/Public/TaggedAssetBrowserFilters/TaggedAssetBrowser_CommonFilters.h b/Engine/Source/Editor/UserAssetTagsEditor/Public/TaggedAssetBrowserFilters/TaggedAssetBrowser_CommonFilters.h
--- a/Engine/Source/Editor/UserAssetTagsEditor/Public/TaggedAssetBrowserFilters/TaggedAssetBrowser_CommonFilters.h
+++ b/Engine/Source/Editor/UserAssetTagsEditor/Public/TaggedAssetBrowserFilters/TaggedAssetBrowser_CommonFilters.h
@@ -149,3 +149,17 @@ public:
 	UPROPERTY(EditAnywhere, Category="Filter", meta=(ContentDir))
 	TArray<FDirectoryPath> DirectoryPaths;
 };
+
+/** Filter for assets based on class. */
+UCLASS(DisplayName="Classes")
+class USERASSETTAGSEDITOR_API UTaggedAssetBrowserFilter_Class : public UTaggedAssetBrowserFilterBase
+{
+	GENERATED_BODY()
+public:
+	UTaggedAssetBrowserFilter_Class() = default;
+	
+	virtual void ModifyARFilterInternal(FARFilter& Filter) const override;
+
+	UPROPERTY(EditAnywhere, Category="Filter")
+	TArray<TObjectPtr<UClass>> Classes;
+};
