Added IsAssetVisible check for tagged asset browser extension functionality

diff --git a/Engine/Source/Editor/UserAssetTagsEditor/Private/Widgets/STaggedAssetBrowser.cpp b/Engine/Source/Editor/UserAssetTagsEditor/Private/Widgets/STaggedAssetBrowser.cpp
--- a/Engine/Source/Editor/UserAssetTagsEditor/Private/Widgets/STaggedAssetBrowser.cpp
+++ b/Engine/Source/Editor/UserAssetTagsEditor/Private/Widgets/STaggedAssetBrowser.cpp
@@ -1125,13 +1125,19 @@ TArray<UTaggedAssetBrowserConfiguration*> STaggedAssetBrowser::GetExtensionAsset
 
 		for(const FAssetData& ExtensionConfigurationAssetData : ExtensionAssets)
 		{
+			// If the asset isn't visible in the user's context, we skip
+			if(IAssetTools::Get().IsAssetVisible(ExtensionConfigurationAssetData, true) == false)
+			{
+				continue;
+			}
+			
 			UTaggedAssetBrowserConfiguration* ExtensionConfigurationAsset = Cast<UTaggedAssetBrowserConfiguration>(ExtensionConfigurationAssetData.GetAsset());
 
 			// The asset should have correct values but we make sure since we need to load it anyways
 			if(ExtensionConfigurationAsset->bIsExtension == false || ExtensionConfigurationAsset->ProfileName != TaggedBrowserProfileName)
 			{
 				continue;
-			}
+			}			
 
 			Result.Add(ExtensionConfigurationAsset);
 		}
diff --git a/Engine/Source/Editor/UserAssetTagsEditor/UserAssetTagsEditor.Build.cs b/Engine/Source/Editor/UserAssetTagsEditor/UserAssetTagsEditor.Build.cs
--- a/Engine/Source/Editor/UserAssetTagsEditor/UserAssetTagsEditor.Build.cs
+++ b/Engine/Source/Editor/UserAssetTagsEditor/UserAssetTagsEditor.Build.cs
@@ -39,7 +39,8 @@ public class UserAssetTagsEditor : ModuleRules
 				"InputCore",
 				"ApplicationCore",
 				"Kismet",
-				"DeveloperSettings"
+				"DeveloperSettings",
+				"AssetTools"
 				// ... add private dependencies that you statically link with here ...	
 			}
 			);
