Added Asset Context filtering to the Tagged Asset Browser.
CreateAsset window will use the "last save directory" path as context.

AddEmitterToSystem window will use the owning system as context.

Tweaks to "SetLastDirectory(NewAsset)" usage so that factories can access that last directory for context reasons during ConfigureProperties

diff --git a/Engine/Plugins/Editor/ContentBrowser/ContentBrowserAssetDataSource/Source/ContentBrowserAssetDataSource/Private/ContentBrowserAssetDataSource.cpp b/Engine/Plugins/Editor/ContentBrowser/ContentBrowserAssetDataSource/Source/ContentBrowserAssetDataSource/Private/ContentBrowserAssetDataSource.cpp
--- a/Engine/Plugins/Editor/ContentBrowser/ContentBrowserAssetDataSource/Source/ContentBrowserAssetDataSource/Private/ContentBrowserAssetDataSource.cpp
+++ b/Engine/Plugins/Editor/ContentBrowser/ContentBrowserAssetDataSource/Source/ContentBrowserAssetDataSource/Private/ContentBrowserAssetDataSource.cpp
@@ -34,6 +34,7 @@
 #include "ContentBrowserModule.h"
 #include "IContentBrowserSingleton.h"
 #include "ContentBrowserMenuContexts.h"
+#include "EditorDirectories.h"
 #include "Misc/Paths.h"
 #include "Widgets/Input/SButton.h"
 #include "SActionButton.h"
@@ -3975,7 +3976,9 @@ void UContentBrowserAssetDataSource::OnNewAssetRequested(const FName InPath, TWe
 
 		// This factory may get gc'd as a side effect of various delegates potentially calling CollectGarbage so protect against it from being gc'd out from under us
 		FGCObjectScopeGuard FactoryGCGuard(NewFactory);
-
+		
+		FEditorDirectories::Get().SetLastDirectory(ELastDirectory::NEW_ASSET, InPath.ToString());
+		
 		FEditorDelegates::OnConfigureNewAssetProperties.Broadcast(NewFactory);
 		if (NewFactory->ConfigureProperties())
 		{
diff --git a/Engine/Plugins/FX/Niagara/Source/NiagaraEditor/Private/NiagaraEmitterFactoryNew.cpp b/Engine/Plugins/FX/Niagara/Source/NiagaraEditor/Private/NiagaraEmitterFactoryNew.cpp
--- a/Engine/Plugins/FX/Niagara/Source/NiagaraEditor/Private/NiagaraEmitterFactoryNew.cpp
+++ b/Engine/Plugins/FX/Niagara/Source/NiagaraEditor/Private/NiagaraEmitterFactoryNew.cpp
@@ -1,6 +1,8 @@
 // Copyright Epic Games, Inc. All Rights Reserved.
 
 #include "NiagaraEmitterFactoryNew.h"
+
+#include "EditorDirectories.h"
 #include "NiagaraEmitter.h"
 #include "NiagaraEditorModule.h"
 #include "NiagaraScriptFactoryNew.h"
@@ -68,7 +70,8 @@ bool UNiagaraEmitterFactoryNew::ConfigureProperties()
 			.RecentAndFavoritesList(FNiagaraEditorModule::Get().GetRecentsManager()->GetRecentEmitterAndSystemsList())
 			.OnAdditionalShouldFilterAsset_UObject(this, &UNiagaraEmitterFactoryNew::OnAdditionalShouldFilterAsset)
 			.OnGetExtraFrontendFilters_UObject(this, &UNiagaraEmitterFactoryNew::OnGetExtraFrontendFilters)
-			.OnExtendAddFilterMenu_UObject(this, &UNiagaraEmitterFactoryNew::OnExtendAddFilterMenu);
+			.OnExtendAddFilterMenu_UObject(this, &UNiagaraEmitterFactoryNew::OnExtendAddFilterMenu)
+			.AdditionalReferencingAssets({ FAssetData("TmpAsset", FName(FEditorDirectories::Get().GetLastDirectory(ELastDirectory::NEW_ASSET)), "TmpAsset", UObject::StaticClass()->GetClassPathName())});
 
 		SWindow::FArguments WindowArgs;
 		WindowArgs.Title(LOCTEXT("EmitterAssetBrowserWindowTitle", "Create Niagara Emitter - Select an emitter as a base"));
diff --git a/Engine/Plugins/FX/Niagara/Source/NiagaraEditor/Private/NiagaraSystemFactoryNew.cpp b/Engine/Plugins/FX/Niagara/Source/NiagaraEditor/Private/NiagaraSystemFactoryNew.cpp
--- a/Engine/Plugins/FX/Niagara/Source/NiagaraEditor/Private/NiagaraSystemFactoryNew.cpp
+++ b/Engine/Plugins/FX/Niagara/Source/NiagaraEditor/Private/NiagaraSystemFactoryNew.cpp
@@ -8,6 +8,7 @@
 #include "NiagaraEditorSettings.h"
 #include "AssetRegistry/AssetData.h"
 #include "ContentBrowserModule.h"
+#include "EditorDirectories.h"
 #include "ImageUtils.h"
 #include "ViewModels/Stack/NiagaraStackGraphUtilities.h"
 #include "Misc/MessageDialog.h"
@@ -64,7 +65,8 @@ bool UNiagaraSystemFactoryNew::ConfigureProperties()
 			.InterfaceOverrideProfiles(OverrideProfiles)
 			.RecentAndFavoritesList(FNiagaraEditorModule::Get().GetRecentsManager()->GetRecentEmitterAndSystemsList())
 			.OnGetExtraFrontendFilters_UObject(this, &UNiagaraSystemFactoryNew::OnGetExtraFrontendFilters)
-			.OnExtendAddFilterMenu_UObject(this, &UNiagaraSystemFactoryNew::OnExtendAddFilterMenu);
+			.OnExtendAddFilterMenu_UObject(this, &UNiagaraSystemFactoryNew::OnExtendAddFilterMenu)
+			.AdditionalReferencingAssets({ FAssetData("TmpAsset", FName(FEditorDirectories::Get().GetLastDirectory(ELastDirectory::NEW_ASSET)), "TmpAsset", UObject::StaticClass()->GetClassPathName())});
 
 		SWindow::FArguments WindowArgs;
 		WindowArgs.Title(LOCTEXT("SystemAssetBrowserWindowTitle", "Create Niagara System - Select an emitter or system as a base"));
diff --git a/Engine/Plugins/FX/Niagara/Source/NiagaraEditor/Private/Widgets/AssetBrowser/SNiagaraAddEmitterToSystemWindow.cpp b/Engine/Plugins/FX/Niagara/Source/NiagaraEditor/Private/Widgets/AssetBrowser/SNiagaraAddEmitterToSystemWindow.cpp
--- a/Engine/Plugins/FX/Niagara/Source/NiagaraEditor/Private/Widgets/AssetBrowser/SNiagaraAddEmitterToSystemWindow.cpp
+++ b/Engine/Plugins/FX/Niagara/Source/NiagaraEditor/Private/Widgets/AssetBrowser/SNiagaraAddEmitterToSystemWindow.cpp
@@ -114,6 +114,7 @@ void SNiagaraAddEmitterToSystemWindow::Construct(const FArguments& InArgs, const
 	AssetBrowserArgs._DefaultDetailsTabConfiguration = DefaultDetailsTabConfiguration;
 	AssetBrowserArgs._OnAssetsActivated = FOnAssetsActivated::CreateSP(this, &SNiagaraAddEmitterToSystemWindow::OnAssetsActivatedInternal);
 	AssetBrowserArgs._RecentAndFavoritesList = FNiagaraEditorModule::Get().GetRecentsManager()->GetRecentEmitterAndSystemsList();
+	AssetBrowserArgs._AdditionalReferencingAssets = { FAssetData(&SystemViewModel->GetSystem()) };
 
 	AssetBrowserArgs._AdditionalBottomWidget = SNew(SBox)
 	   .Padding(16.f, 16.f)
diff --git a/Engine/Source/Editor/UserAssetTagsEditor/Private/Widgets/STaggedAssetBrowser.cpp b/Engine/Source/Editor/UserAssetTagsEditor/Private/Widgets/STaggedAssetBrowser.cpp
--- a/Engine/Source/Editor/UserAssetTagsEditor/Private/Widgets/STaggedAssetBrowser.cpp
+++ b/Engine/Source/Editor/UserAssetTagsEditor/Private/Widgets/STaggedAssetBrowser.cpp
@@ -44,6 +44,8 @@ void STaggedAssetBrowser::Construct(const FArguments& InArgs, const UTaggedAsset
 	InterfaceOverrideProfiles = InArgs._InterfaceOverrideProfiles;
 	OnGenerateAssetTooltipWidgetDelegate = InArgs._OnGenerateAssetTooltipWidget;
 	OnIsAssetValidForCustomToolTipDelegate = InArgs._OnIsAssetValidForCustomToolTip;
+	PropertyHandle = InArgs._PropertyHandle;
+	AdditionalReferencingAssets = InArgs._AdditionalReferencingAssets;
 
 	CheckValidSetup();
 	
@@ -495,6 +497,8 @@ void STaggedAssetBrowser::PopulateAssetBrowserContentSlot()
 	Config.bAllowDragging = false;
 	Config.AssetViewOptionsProfile = InterfaceOverrideProfiles.AssetViewOptionsProfileName.IsSet() ? InterfaceOverrideProfiles.AssetViewOptionsProfileName.GetValue() : TaggedBrowserProfileName;
 	Config.OnGetAssetContextMenu = FOnGetAssetContextMenu::CreateSP(this, &STaggedAssetBrowser::OnGetAssetContextMenu);
+	Config.PropertyHandle = PropertyHandle;
+	Config.AdditionalReferencingAssets = AdditionalReferencingAssets;
 
 	if(TaggedBrowserProfileName.IsSet())
 	{
@@ -1179,7 +1183,9 @@ void STaggedAssetBrowserWindow::Construct(const FArguments& InArgs, const UTagge
 		.OnAssetsActivated(this, &STaggedAssetBrowserWindow::OnAssetsActivated)
 		.OnGenerateAssetTooltipWidget(InArgs._AssetBrowserArgs._OnGenerateAssetTooltipWidget)
 		.OnIsAssetValidForCustomToolTip(InArgs._AssetBrowserArgs._OnIsAssetValidForCustomToolTip)
-		.AdditionalBottomWidget(InArgs._AssetBrowserArgs._AdditionalBottomWidget);
+		.AdditionalBottomWidget(InArgs._AssetBrowserArgs._AdditionalBottomWidget)
+		.PropertyHandle(InArgs._AssetBrowserArgs._PropertyHandle)
+		.AdditionalReferencingAssets(InArgs._AssetBrowserArgs._AdditionalReferencingAssets);
 	
 	SetContent(AssetBrowser.ToSharedRef());
 }
diff --git a/Engine/Source/Editor/UserAssetTagsEditor/Private/Widgets/STaggedAssetBrowserCreateAssetWindow.cpp b/Engine/Source/Editor/UserAssetTagsEditor/Private/Widgets/STaggedAssetBrowserCreateAssetWindow.cpp
--- a/Engine/Source/Editor/UserAssetTagsEditor/Private/Widgets/STaggedAssetBrowserCreateAssetWindow.cpp
+++ b/Engine/Source/Editor/UserAssetTagsEditor/Private/Widgets/STaggedAssetBrowserCreateAssetWindow.cpp
@@ -1,6 +1,9 @@
 ﻿// Copyright Epic Games, Inc. All Rights Reserved.
 
 #include "Widgets/STaggedAssetBrowserCreateAssetWindow.h"
+
+#include "ContentBrowserModule.h"
+#include "EditorDirectories.h"
 #include "Widgets/Layout/SSpacer.h"
 
 #define LOCTEXT_NAMESPACE "TaggedAssetBrowser"
@@ -85,6 +88,13 @@ void STaggedAssetBrowserCreateAssetWindow::Construct(const FArguments& InArgs, c
 	
 	Args._AssetBrowserWindowArgs._AssetBrowserArgs._AdditionalBottomWidget = CreateAssetControls;
 	Args._AssetBrowserWindowArgs._AssetBrowserArgs._OnAssetsActivated = FOnAssetsActivated::CreateSP(this, &STaggedAssetBrowserCreateAssetWindow::OnAssetsActivatedInternal);
+
+	// If we haven't specified any referencing assets, we assume the context is the last "save asset" editor path, which should be set before this is called
+	if(Args._AssetBrowserWindowArgs._AssetBrowserArgs._AdditionalReferencingAssets.Num() == 0)
+	{
+		Args._AssetBrowserWindowArgs._AssetBrowserArgs._AdditionalReferencingAssets = { FAssetData("TmpAsset", FName(FEditorDirectories::Get().GetLastDirectory(ELastDirectory::NEW_ASSET)), "TmpAsset", UObject::StaticClass()->GetClassPathName()) };  
+	}
+	
 	STaggedAssetBrowserWindow::Construct(Args._AssetBrowserWindowArgs, Configuration);
 }
 
diff --git a/Engine/Source/Editor/UserAssetTagsEditor/Public/Widgets/STaggedAssetBrowser.h b/Engine/Source/Editor/UserAssetTagsEditor/Public/Widgets/STaggedAssetBrowser.h
--- a/Engine/Source/Editor/UserAssetTagsEditor/Public/Widgets/STaggedAssetBrowser.h
+++ b/Engine/Source/Editor/UserAssetTagsEditor/Public/Widgets/STaggedAssetBrowser.h
@@ -76,6 +76,10 @@ public:
 		SLATE_ARGUMENT(FInterfaceOverrideProfiles, InterfaceOverrideProfiles)
 		/** An externally owned recent & favorites list. Defaults to the editor's main list, but can be a custom list as well. */
 		SLATE_ATTRIBUTE(const FMainMRUFavoritesList*, RecentAndFavoritesList)
+		/** The property handle for which this browser is summoned. Used for contextual filtering. In wizard windows or similar, will be nullptr. */
+		SLATE_ARGUMENT(TSharedPtr<IPropertyHandle>, PropertyHandle)
+		/** Additional assets used for contextual filtering. In wizard windows or similar, construct FAssetData pointing at the relevant path. */
+		SLATE_ARGUMENT(TArray<FAssetData>, AdditionalReferencingAssets)
 		
 		/** Additional Content Browser filters (in the dropdown) can specified here. */
 		SLATE_EVENT(FOnGetExtraFrontendFilters, OnGetExtraFrontendFilters)
@@ -221,6 +225,8 @@ private:
 	FOnExtendAddFilterMenu OnExtendAddFilterMenuDelegate;
 	FOnShouldFilterAsset OnAdditionalShouldFilterAssetDelegate;
 	ESelectionMode::Type AssetSelectionMode = ESelectionMode::Single;
+	TSharedPtr<IPropertyHandle> PropertyHandle;
+	TArray<FAssetData> AdditionalReferencingAssets;
 
 	TArray<FFilterPath> SourceSearchResults;
 	TOptional<FFilterPath> FocusedSearchResult;
