Minor bugfixes Tagged Asset Browser

diff --git a/Engine/Source/Editor/UserAssetTagsEditor/Private/Widgets/STaggedAssetBrowser.cpp b/Engine/Source/Editor/UserAssetTagsEditor/Private/Widgets/STaggedAssetBrowser.cpp
--- a/Engine/Source/Editor/UserAssetTagsEditor/Private/Widgets/STaggedAssetBrowser.cpp
+++ b/Engine/Source/Editor/UserAssetTagsEditor/Private/Widgets/STaggedAssetBrowser.cpp
@@ -8,6 +8,7 @@
 #include "IContentBrowserDataModule.h"
 #include "SAssetView.h"
 #include "SlateOptMacros.h"
+#include "UserAssetTagEditorMenuContexts.h"
 #include "AssetRegistry/AssetRegistryModule.h"
 #include "Assets/TaggedAssetBrowserConfiguration.h"
 #include "Logging/StructuredLog.h"
@@ -27,6 +28,9 @@ void STaggedAssetBrowser::Construct(const FArguments& InArgs, const UTaggedAsset
 	Configuration = &InTaggedAssetBrowserConfiguration;
 	TaggedBrowserProfileName = Configuration->ProfileName;
 
+	ThisContext.Reset(NewObject<UTaggedAssetBrowserMenuContext>());
+	ThisContext->TaggedAssetBrowser = SharedThis(this).ToWeakPtr();
+	
 	// We duplicate the root here as we manipulate data of the filters.
 	FilterHierarchyRoot.Reset(Cast<UTaggedAssetBrowserFilterRoot>(StaticDuplicateObject(Configuration->FilterRoot, GetTransientPackage())));
 
@@ -147,8 +151,6 @@ void STaggedAssetBrowser::InitializeValues()
 			return false;
 		});
 	}
-
-	EnableToolTipForceField(true);
 }
 
 void STaggedAssetBrowser::RefreshPrimaryFilterRootItems(bool bRefreshWidget)
@@ -243,9 +245,15 @@ void STaggedAssetBrowser::InitAssetContextMenu()
 			}
 		});
 
-		FToolMenuCanExecuteAction CanBrowseToAssetAction = FToolMenuCanExecuteAction::CreateLambda([this](const FToolMenuContext& Context)
+		FToolMenuCanExecuteAction CanBrowseToAssetAction = FToolMenuCanExecuteAction::CreateLambda([](const FToolMenuContext& Context)
 		{
-			TSharedPtr<SWindow> Window = FSlateApplication::Get().FindWidgetWindow(AsShared());
+			UTaggedAssetBrowserMenuContext* TaggedAssetBrowserMenuContext = Context.FindContext<UTaggedAssetBrowserMenuContext>();
+			if (TaggedAssetBrowserMenuContext == nullptr || !TaggedAssetBrowserMenuContext->TaggedAssetBrowser.IsValid())
+			{
+				return false;
+			}
+				
+			TSharedPtr<SWindow> Window = FSlateApplication::Get().FindWidgetWindow(TaggedAssetBrowserMenuContext->TaggedAssetBrowser.Pin().ToSharedRef());
 			if(Window && Window->IsModalWindow())
 			{
 				return false;
@@ -471,6 +479,8 @@ void STaggedAssetBrowser::PopulatePrimaryFiltersSlot()
 		[
 			PrimaryFilterSelector.ToSharedRef()	
 		];
+
+	PrimaryFiltersSlotWidget->EnableToolTipForceField(true);
 	
 	PrimaryFiltersSlot->AttachWidget(PrimaryFiltersSlotWidget.ToSharedRef());
 }
@@ -957,9 +967,12 @@ TSharedPtr<SWidget> STaggedAssetBrowser::OnGetAssetContextMenu(const TArray<FAss
 {
 	if(AssetData.Num() == 1)
 	{
-		UContentBrowserAssetContextMenuContext* MenuContext = NewObject<UContentBrowserAssetContextMenuContext>();
-		MenuContext->SelectedAssets = AssetData;
-		FToolMenuContext Context(MenuContext);
+		UContentBrowserAssetContextMenuContext* ContentBrowserContext = NewObject<UContentBrowserAssetContextMenuContext>();
+		ContentBrowserContext->SelectedAssets = AssetData;
+		
+		FToolMenuContext Context;
+		Context.AddObject(ContentBrowserContext);
+		Context.AddObject(ThisContext.Get());
 		return UToolMenus::Get()->GenerateWidget(GetAssetContextMenuName(), Context);
 	}
 
diff --git a/Engine/Source/Editor/UserAssetTagsEditor/Public/UserAssetTagEditorMenuContexts.h b/Engine/Source/Editor/UserAssetTagsEditor/Public/UserAssetTagEditorMenuContexts.h
--- a/Engine/Source/Editor/UserAssetTagsEditor/Public/UserAssetTagEditorMenuContexts.h
+++ b/Engine/Source/Editor/UserAssetTagsEditor/Public/UserAssetTagEditorMenuContexts.h
@@ -5,11 +5,20 @@
 #include "AssetRegistry/AssetData.h"
 #include "UserAssetTagEditorMenuContexts.generated.h"
 
-UCLASS()
-class USERASSETTAGSEDITOR_API UUserAssetTagEditorContext : public UObject
+UCLASS(MinimalAPI)
+class UUserAssetTagEditorContext : public UObject
 {
 	GENERATED_BODY()
 
 public:
 	TWeakPtr<class SUserAssetTagsEditor> UserAssetTagsEditor;
+};
+
+UCLASS(MinimalAPI)
+class UTaggedAssetBrowserMenuContext : public UObject
+{
+	GENERATED_BODY()
+
+public:
+	TWeakPtr<class STaggedAssetBrowser> TaggedAssetBrowser;
 };
\ No newline at end of file
diff --git a/Engine/Source/Editor/UserAssetTagsEditor/Public/Widgets/STaggedAssetBrowser.h b/Engine/Source/Editor/UserAssetTagsEditor/Public/Widgets/STaggedAssetBrowser.h
--- a/Engine/Source/Editor/UserAssetTagsEditor/Public/Widgets/STaggedAssetBrowser.h
+++ b/Engine/Source/Editor/UserAssetTagsEditor/Public/Widgets/STaggedAssetBrowser.h
@@ -17,6 +17,9 @@
 
 #define UE_API USERASSETTAGSEDITOR_API
 
+class UTaggedAssetBrowserMenuContext;
+struct FAssetPickerConfig;
+
 DECLARE_DELEGATE_RetVal_OneParam(TSharedRef<SWidget>, FOnGenerateAssetWidget, const FAssetData& Asset)
 
 /** You can provide custom tabs to be displayed in the right column. Useful if you want to add custom content. */
@@ -244,6 +247,8 @@ private:
 	TSharedPtr<class STaggedAssetBrowserSections> SectionsSlotWidget;
 	TSharedPtr<SWidget> PrimaryFiltersSlotWidget;
 
+	TStrongObjectPtr<UTaggedAssetBrowserMenuContext> ThisContext;
+
 	struct FDefaultDetailsTabMembers : public TSharedFromThis<FDefaultDetailsTabMembers>
 	{
 		TSharedPtr<SWidgetSwitcher> DefaultDetailsSwitcher;
