UE_LOGFMT usage fixes where %s should be {VariableA}
#rnx

diff --git a/Engine/Plugins/Runtime/Metasound/Source/MetasoundEngine/Private/MetasoundFrontendLiteralBlueprintAccess.cpp b/Engine/Plugins/Runtime/Metasound/Source/MetasoundEngine/Private/MetasoundFrontendLiteralBlueprintAccess.cpp
--- a/Engine/Plugins/Runtime/Metasound/Source/MetasoundEngine/Private/MetasoundFrontendLiteralBlueprintAccess.cpp
+++ b/Engine/Plugins/Runtime/Metasound/Source/MetasoundEngine/Private/MetasoundFrontendLiteralBlueprintAccess.cpp
@@ -41,7 +41,7 @@ namespace MetasoundFrontendLiteralBlueprintAccessPrivate
 
 		if (Result != EMetaSoundBuilderResult::Succeeded)
 		{
-			UE_LOGFMT(LogMetaSound, Warning, "Failed to get %s value from InLiteral.", *UEnum::GetValueAsString(ExpectedType));
+			UE_LOGFMT(LogMetaSound, Warning, "Failed to get {ExpectedType} value from InLiteral.", *UEnum::GetValueAsString(ExpectedType));
 			return DefaultValue;
 		}
 
diff --git a/Engine/Source/Editor/UserAssetTagsEditor/Private/DataHierarchyEditor/TaggedAssetBrowserConfigurationHierarchyViewModel.cpp b/Engine/Source/Editor/UserAssetTagsEditor/Private/DataHierarchyEditor/TaggedAssetBrowserConfigurationHierarchyViewModel.cpp
--- a/Engine/Source/Editor/UserAssetTagsEditor/Private/DataHierarchyEditor/TaggedAssetBrowserConfigurationHierarchyViewModel.cpp
+++ b/Engine/Source/Editor/UserAssetTagsEditor/Private/DataHierarchyEditor/TaggedAssetBrowserConfigurationHierarchyViewModel.cpp
@@ -93,7 +93,7 @@ TArray<TSubclassOf<UTaggedAssetBrowserFilterBase>> UTaggedAssetBrowserConfigurat
 		
 		if(Candidate->HasAnyClassFlags(CLASS_Abstract | CLASS_Deprecated))
 		{
-			UE_LOGFMT(LogUserAssetTags, Warning, "Abstract or deprecated class found in list of valid default filter classes in %s", Candidate->GetName());
+			UE_LOGFMT(LogUserAssetTags, Warning, "Abstract or deprecated class found in list of valid default filter classes in {Candidate}", Candidate->GetName());
 			return true;
 		}
 
