[TagAssetBrowser] Fix missing UE_API export on STaggedAssetBrowser::Construct

diff --git a/Engine/Source/Editor/UserAssetTagsEditor/Public/Widgets/STaggedAssetBrowser.h b/Engine/Source/Editor/UserAssetTagsEditor/Public/Widgets/STaggedAssetBrowser.h
--- a/Engine/Source/Editor/UserAssetTagsEditor/Public/Widgets/STaggedAssetBrowser.h
+++ b/Engine/Source/Editor/UserAssetTagsEditor/Public/Widgets/STaggedAssetBrowser.h
@@ -103,7 +103,7 @@ public:
 	SLATE_END_ARGS()
 
 	
-	void Construct(const FArguments& InArgs, const UTaggedAssetBrowserConfiguration& InTaggedAssetBrowserConfiguration);
+	UE_API void Construct(const FArguments& InArgs, const UTaggedAssetBrowserConfiguration& InTaggedAssetBrowserConfiguration);
 
 	UE_API STaggedAssetBrowser();
 	UE_API virtual ~STaggedAssetBrowser() override;
