UAF Browser Tagged Asset Browser Changes
- Allow custom tabs to be ignored.
- Allow Asset Picker args to be generally externally customized.
- Make class filters show first class icon & name by default
- Make directory filters be named by directory by default
- Add optional name for class & directory filter names
- Remove requirement for at least one available class filter, as it's possible to not have any available classes to allow all classes.

#tests Editor

diff --git a/Engine/Source/Editor/UserAssetTagsEditor/Private/TaggedAssetBrowserFilters/TaggedAssetBrowser_CommonFilters.cpp b/Engine/Source/Editor/UserAssetTagsEditor/Private/TaggedAssetBrowserFilters/TaggedAssetBrowser_CommonFilters.cpp
--- a/Engine/Source/Editor/UserAssetTagsEditor/Private/TaggedAssetBrowserFilters/TaggedAssetBrowser_CommonFilters.cpp
+++ b/Engine/Source/Editor/UserAssetTagsEditor/Private/TaggedAssetBrowserFilters/TaggedAssetBrowser_CommonFilters.cpp
@@ -6,6 +6,7 @@
 #include "UserAssetTagEditorUtilities.h"
 #include "UserAssetTagsEditorModule.h"
 #include "Logging/StructuredLog.h"
+#include "Styling/SlateIconFinder.h"
 #include "Widgets/STaggedAssetBrowser.h"
 #include "Engine/Texture2D.h"
 
@@ -222,6 +223,30 @@ bool UTaggedAssetBrowserFilter_Recent::IsAssetRecent(const FAssetData& AssetCand
 	return false;
 }
 
+FString UTaggedAssetBrowserFilter_Directories::ToString() const
+{
+	if (!FilterName.IsNone())
+	{
+		return FilterName.ToString();
+	}
+
+	if (!DirectoryPaths.IsEmpty())
+	{
+		const FDirectoryPath& FirstDirectoryPath = DirectoryPaths[0];
+
+		int32 LeafIndex = INDEX_NONE;
+		if (FirstDirectoryPath.Path.FindLastChar(TEXT('/'), LeafIndex))
+		{
+			if (LeafIndex != INDEX_NONE)
+			{
+				return FirstDirectoryPath.Path.RightChop(LeafIndex + 1);
+			}
+		}
+	}
+
+	return UTaggedAssetBrowserFilterBase::ToString();
+}
+
 void UTaggedAssetBrowserFilter_Directories::ModifyARFilterInternal(FARFilter& Filter) const
 {
 	// The AR would use regular package paths, but in the context of the Content Browser/Asset Picker, which uses content sources and virtual paths
@@ -239,6 +264,24 @@ void UTaggedAssetBrowserFilter_Directories::ModifyARFilterInternal(FARFilter& Fi
 	Filter.PackagePaths.Append(VirtualPackagePaths);
 }
 
+FText UTaggedAssetBrowserFilter_Directories::GetTooltip() const
+{
+	if (!DirectoryPaths.IsEmpty())
+	{
+		FText BaseTooltip = LOCTEXT("UTaggedAssetBrowserFilter_DirectoriesTooltip", "Directories: \n {0}.");
+		FTextBuilder TooltipListBuilder;
+
+		for (const FDirectoryPath& DirectoryPath : DirectoryPaths)
+		{
+			TooltipListBuilder.AppendLine(DirectoryPath.Path);
+		}
+
+		return FText::FormatOrdered(BaseTooltip, TooltipListBuilder.ToText());
+	}
+
+	return UTaggedAssetBrowserFilterBase::GetTooltip();
+}
+
 FSlateIcon UTaggedAssetBrowserFilter_Directories::GetIcon() const
 {
 	if(ActiveContext.IsSet())
@@ -252,6 +295,60 @@ FSlateIcon UTaggedAssetBrowserFilter_Directories::GetIcon() const
 	return FSlateIcon(FAppStyle::GetAppStyleSetName(), "Icons.FolderClosed");
 }
 
+FString UTaggedAssetBrowserFilter_Class::ToString() const
+{
+	if (!FilterName.IsNone())
+	{
+		return FilterName.ToString();
+	}
+
+	if (!Classes.IsEmpty())
+	{
+		const UClass* FirstClass = Classes[0];
+		if (FirstClass)
+		{
+			return FirstClass->GetDisplayNameText().ToString();
+		}
+	}
+
+	return UTaggedAssetBrowserFilterBase::ToString();
+}
+
+FText UTaggedAssetBrowserFilter_Class::GetTooltip() const
+{
+	if (!Classes.IsEmpty())
+	{
+		FText BaseTooltip = LOCTEXT("UTaggedAssetBrowserFilter_ClassTooltip", "Classes: \n {0}.");
+		FTextBuilder TooltipListBuilder;
+
+		for (const UClass* Class : Classes)
+		{
+			if (Class)
+			{
+				TooltipListBuilder.AppendLine(Class->GetDisplayNameText().ToString());
+			}
+		}
+
+		return FText::FormatOrdered(BaseTooltip, TooltipListBuilder.ToText());
+	}
+
+	return UTaggedAssetBrowserFilterBase::GetTooltip();
+}
+
+FSlateIcon UTaggedAssetBrowserFilter_Class::GetIcon() const
+{
+	if (!Classes.IsEmpty())
+	{
+		const UClass* FirstClass = Classes[0];
+		if (FirstClass)
+		{
+			return FSlateIconFinder::FindIconForClass(FirstClass);
+		}
+	}
+
+	return UTaggedAssetBrowserFilterBase::GetIcon();
+}
+
 void UTaggedAssetBrowserFilter_Class::ModifyARFilterInternal(FARFilter& Filter) const
 {
 	for(UClass* Class : Classes)
diff --git a/Engine/Source/Editor/UserAssetTagsEditor/Private/Widgets/STaggedAssetBrowser.cpp b/Engine/Source/Editor/UserAssetTagsEditor/Private/Widgets/STaggedAssetBrowser.cpp
--- a/Engine/Source/Editor/UserAssetTagsEditor/Private/Widgets/STaggedAssetBrowser.cpp
+++ b/Engine/Source/Editor/UserAssetTagsEditor/Private/Widgets/STaggedAssetBrowser.cpp
@@ -41,6 +41,7 @@ void STaggedAssetBrowser::Construct(const FArguments& InArgs, const UTaggedAsset
 	OnAssetsActivatedDelegate = InArgs._OnAssetsActivated;
 	OnGetExtraFrontendFiltersDelegate = InArgs._OnGetExtraFrontendFilters;
 	OnExtendAddFilterMenuDelegate = InArgs._OnExtendAddFilterMenu;
+	OnOverrideAssetPickerConfigDelegate = InArgs._OnOverrideAssetPickerConfig;
 	InterfaceOverrideProfiles = InArgs._InterfaceOverrideProfiles;
 	OnGenerateAssetTooltipWidgetDelegate = InArgs._OnGenerateAssetTooltipWidget;
 	OnIsAssetValidForCustomToolTipDelegate = InArgs._OnIsAssetValidForCustomToolTip;
@@ -50,7 +51,7 @@ void STaggedAssetBrowser::Construct(const FArguments& InArgs, const UTaggedAsset
 	CheckValidSetup();
 	
 	InitializeValues();
-	
+
 	ChildSlot
 	[
 		SNew(SBorder)
@@ -64,21 +65,7 @@ void STaggedAssetBrowser::Construct(const FArguments& InArgs, const UTaggedAsset
 				SNew(SHorizontalBox)
 				+ SHorizontalBox::Slot().Expose(SectionsSlot)
 				.AutoWidth()
-				+ SHorizontalBox::Slot()
-				[
-					SNew(SSplitter)
-					.Orientation(Orient_Horizontal)
-					.PhysicalSplitterHandleSize(2.f)
-					+ SSplitter::Slot().Expose(PrimaryFiltersSlot)
-						.Value(this, &STaggedAssetBrowser::OnGetPrimaryFiltersSlotSizeValue)
-						.MinSize(this, &STaggedAssetBrowser::OnGetPrimaryFiltersSlotMinSize)
-						.OnSlotResized(this, &STaggedAssetBrowser::OnPrimaryFiltersSlotResized)
-					+ SSplitter::Slot().Expose(AssetBrowserContentSlot)
-						.Value(0.6f)
-					+ SSplitter::Slot().Expose(CustomTabsSlot)
-						.Value(0.25f)
-						.MinSize(100.f)
-				]
+				+ SHorizontalBox::Slot().Expose(AssetBrowserAreaSplitterSlot)
 			]
 			+ SVerticalBox::Slot()
 			.AutoHeight()
@@ -94,10 +81,15 @@ void STaggedAssetBrowser::Construct(const FArguments& InArgs, const UTaggedAsset
 	];
 
 	bSuppressSaveAndLoad = true;
+	PopulateAssetBrowserSplitterSlot();
 	PopulateAssetBrowserContentSlot();
 	PopulateSectionsSlot();
 	PopulatePrimaryFiltersSlot();
-	PopulateCustomTabsSlot();
+
+	if (NeedsCustomTabArea())
+	{
+		PopulateCustomTabsSlot();
+	}
 
 	if(InArgs._AdditionalBottomWidget.IsValid())
 	{
@@ -274,8 +266,6 @@ void STaggedAssetBrowser::InitAssetContextMenu()
 
 void STaggedAssetBrowser::CheckValidSetup() const
 {
-	ensureMsgf(AvailableClasses.Num() >= 1, TEXT("The Tagged Asset Browser has to be supplied at least one available class."));
-
 	if(OnIsAssetValidForCustomToolTipDelegate.IsBound())
 	{
 		ensureMsgf(OnGenerateAssetTooltipWidgetDelegate.IsBound(), TEXT("If IsAssetValidForCustomToolTip is bound, you also have to bind OnGenerateAssetTooltipWidget."));
@@ -395,6 +385,11 @@ bool STaggedAssetBrowser::ShouldFilterAsset(const FAssetData& AssetData) const
 	return false;
 }
 
+bool STaggedAssetBrowser::NeedsCustomTabArea() const
+{
+	return DefaultDetailsTabConfiguration.bUseDefaultDetailsTab || !CustomTabInfos.IsEmpty();
+}
+
 void STaggedAssetBrowser::RefreshBackendFilter() const
 {
 	AssetBrowserContent->SetARFilter(CreateCurrentBackendFilter());
@@ -475,6 +470,40 @@ void STaggedAssetBrowser::PopulatePrimaryFiltersSlot()
 	PrimaryFiltersSlot->AttachWidget(PrimaryFiltersSlotWidget.ToSharedRef());
 }
 
+void STaggedAssetBrowser::PopulateAssetBrowserSplitterSlot()
+{
+	// @TODO: The splitter logic here is duplicated since SSplitter currently doesn't support 
+	// after creation slot modification like horizontal / vertical box do.
+
+	if (NeedsCustomTabArea())
+	{
+		AssetBrowserAreaSplitterSlot->AttachWidget(SNew(SSplitter)
+			.Orientation(Orient_Horizontal)
+			.PhysicalSplitterHandleSize(2.f)
+			+ SSplitter::Slot().Expose(PrimaryFiltersSlot)
+			.Value(this, &STaggedAssetBrowser::OnGetPrimaryFiltersSlotSizeValue)
+			.MinSize(this, &STaggedAssetBrowser::OnGetPrimaryFiltersSlotMinSize)
+			.OnSlotResized(this, &STaggedAssetBrowser::OnPrimaryFiltersSlotResized)
+			+ SSplitter::Slot().Expose(AssetBrowserContentSlot)
+			.Value(0.6f)
+			+ SSplitter::Slot().Expose(CustomTabsSlot)
+			.Value(0.25f)
+			.MinSize(100.f));
+	}
+	else
+	{
+		AssetBrowserAreaSplitterSlot->AttachWidget(SNew(SSplitter)
+			.Orientation(Orient_Horizontal)
+			.PhysicalSplitterHandleSize(2.f)
+			+ SSplitter::Slot().Expose(PrimaryFiltersSlot)
+			.Value(this, &STaggedAssetBrowser::OnGetPrimaryFiltersSlotSizeValue)
+			.MinSize(this, &STaggedAssetBrowser::OnGetPrimaryFiltersSlotMinSize)
+			.OnSlotResized(this, &STaggedAssetBrowser::OnPrimaryFiltersSlotResized)
+			+ SSplitter::Slot().Expose(AssetBrowserContentSlot)
+			.Value(0.85f));
+	}
+}
+
 void STaggedAssetBrowser::PopulateAssetBrowserContentSlot()
 {
 	FAssetPickerConfig Config;
@@ -506,7 +535,9 @@ void STaggedAssetBrowser::PopulateAssetBrowserContentSlot()
 	}
 
 	Config.FilterBarSaveSettingsNameOverride = InterfaceOverrideProfiles.FilterBarSaveName.IsSet() ? InterfaceOverrideProfiles.FilterBarSaveName.GetValue().ToString() : TOptional<FString>();
-	
+
+	OnOverrideAssetPickerConfigDelegate.ExecuteIfBound(Config);
+
 	SAssignNew(AssetBrowserContent, STaggedAssetBrowserContent)
 	.InitialConfig(Config);
 	
@@ -614,7 +645,10 @@ void STaggedAssetBrowser::PopulateCustomTabsSlot()
 			TabSwitcher
 		];
 	
-	CustomTabsSlot->AttachWidget(Widget);
+	if (!CustomTabInfos.IsEmpty())
+	{
+		CustomTabsSlot->AttachWidget(Widget);
+	}
 }
 
 void STaggedAssetBrowser::OnPrimaryFiltersSlotResized(float InSize)
diff --git a/Engine/Source/Editor/UserAssetTagsEditor/Public/TaggedAssetBrowserFilters/TaggedAssetBrowser_CommonFilters.h b/Engine/Source/Editor/UserAssetTagsEditor/Public/TaggedAssetBrowserFilters/TaggedAssetBrowser_CommonFilters.h
--- a/Engine/Source/Editor/UserAssetTagsEditor/Public/TaggedAssetBrowserFilters/TaggedAssetBrowser_CommonFilters.h
+++ b/Engine/Source/Editor/UserAssetTagsEditor/Public/TaggedAssetBrowserFilters/TaggedAssetBrowser_CommonFilters.h
@@ -142,24 +142,40 @@ class USERASSETTAGSEDITOR_API UTaggedAssetBrowserFilter_Directories : public UTa
 	GENERATED_BODY()
 public:
 	UTaggedAssetBrowserFilter_Directories() = default;
-	
+
+	virtual FString ToString() const override;
+	virtual TOptional<FString> GetInstanceIdentifier() const override { return ToString(); }
 	virtual void ModifyARFilterInternal(FARFilter& Filter) const override;
+	virtual FText GetTooltip() const override;
 	virtual FSlateIcon GetIcon() const override;
 
 	UPROPERTY(EditAnywhere, Category="Filter", meta=(ContentDir))
 	TArray<FDirectoryPath> DirectoryPaths;
+
+	/** Name for this directory filter, by default will use first directory name if not set */
+	UPROPERTY(EditAnywhere, Category = "Filter")
+	FName FilterName;
 };
 
-/** Filter for assets based on class. */
+/** Filter for assets based on class. By default will use first class as an icon. */
 UCLASS(DisplayName="Classes")
 class USERASSETTAGSEDITOR_API UTaggedAssetBrowserFilter_Class : public UTaggedAssetBrowserFilterBase
 {
 	GENERATED_BODY()
 public:
 	UTaggedAssetBrowserFilter_Class() = default;
+
+	virtual FString ToString() const override;
+	virtual TOptional<FString> GetInstanceIdentifier() const override { return ToString(); }
+	virtual FText GetTooltip() const override;
+	virtual FSlateIcon GetIcon() const override;
 	
 	virtual void ModifyARFilterInternal(FARFilter& Filter) const override;
 
 	UPROPERTY(EditAnywhere, Category="Filter")
 	TArray<TObjectPtr<UClass>> Classes;
+
+	/** Name for this class filter, by default will use first class in list as name if not set */
+	UPROPERTY(EditAnywhere, Category = "Filter")
+	FName FilterName;
 };
diff --git a/Engine/Source/Editor/UserAssetTagsEditor/Public/Widgets/STaggedAssetBrowser.h b/Engine/Source/Editor/UserAssetTagsEditor/Public/Widgets/STaggedAssetBrowser.h
--- a/Engine/Source/Editor/UserAssetTagsEditor/Public/Widgets/STaggedAssetBrowser.h
+++ b/Engine/Source/Editor/UserAssetTagsEditor/Public/Widgets/STaggedAssetBrowser.h
@@ -17,6 +17,8 @@
 
 #define UE_API USERASSETTAGSEDITOR_API
 
+struct FAssetPickerConfig;
+
 DECLARE_DELEGATE_RetVal_OneParam(TSharedRef<SWidget>, FOnGenerateAssetWidget, const FAssetData& Asset)
 
 /** You can provide custom tabs to be displayed in the right column. Useful if you want to add custom content. */
@@ -50,6 +52,7 @@ class STaggedAssetBrowser : public SCompoundWidget
 {
 public:
 	DECLARE_DELEGATE_RetVal(TArray<TSharedRef<FFrontendFilter>>, FOnGetExtraFrontendFilters)
+	DECLARE_DELEGATE_OneParam(FOnOverrideAssetPickerConfig, FAssetPickerConfig&)
 
 	/** Contains override names for specific menu names. By default, the ProfileName is used as suffix for all menus. */
 	struct FInterfaceOverrideProfiles
@@ -87,6 +90,13 @@ public:
 		SLATE_EVENT(FOnExtendAddFilterMenu, OnExtendAddFilterMenu)
 		/** An additional, always active filter */
 		SLATE_EVENT(FOnShouldFilterAsset, OnAdditionalShouldFilterAsset)
+
+		/**
+		 * Generic Callback for editors to customize our asset picker config before we create it. So that we don't have to update our API every time the asset picker changes.
+		 * Use with care, as we do not guarantee STaggedAssetBrowser functionality will work depending on customizations done. 
+		 * Consider capturing existing values if previous functionality may be important to you.
+		 */
+		SLATE_EVENT(FOnOverrideAssetPickerConfig, OnOverrideAssetPickerConfig)
 		
 		/** If supplied, will generate a custom tooltip when hovering over an asset. */
 		SLATE_EVENT(FOnGenerateAssetWidget, OnGenerateAssetTooltipWidget)
@@ -130,9 +140,11 @@ private:
 	FARFilter CreateCurrentBackendFilter() const;
 	
 	bool ShouldFilterAsset(const FAssetData& AssetData) const;
+	bool NeedsCustomTabArea() const;
 
 	void PopulateSectionsSlot();
 	void PopulatePrimaryFiltersSlot();
+	void PopulateAssetBrowserSplitterSlot();
 	void PopulateAssetBrowserContentSlot();
 	void PopulateCustomTabsSlot();
 
@@ -224,6 +236,7 @@ private:
 	FOnGetExtraFrontendFilters OnGetExtraFrontendFiltersDelegate;
 	FOnExtendAddFilterMenu OnExtendAddFilterMenuDelegate;
 	FOnShouldFilterAsset OnAdditionalShouldFilterAssetDelegate;
+	FOnOverrideAssetPickerConfig OnOverrideAssetPickerConfigDelegate;
 	ESelectionMode::Type AssetSelectionMode = ESelectionMode::Single;
 	TSharedPtr<IPropertyHandle> PropertyHandle;
 	TArray<FAssetData> AdditionalReferencingAssets;
@@ -235,8 +248,9 @@ private:
 	TSharedPtr<STreeView<TObjectPtr<UHierarchyElement>>> PrimaryFilterSelector;
 	TSharedPtr<STaggedAssetBrowserContent> AssetBrowserContent;
 
-	SHorizontalBox::FSlot* SectionsSlot = nullptr;
 	SSplitter::FSlot* PrimaryFiltersSlot = nullptr;
+	SHorizontalBox::FSlot* SectionsSlot = nullptr;
+	SHorizontalBox::FSlot* AssetBrowserAreaSplitterSlot = nullptr;
 	SSplitter::FSlot* AssetBrowserContentSlot = nullptr;
 	SSplitter::FSlot* CustomTabsSlot = nullptr;
 	SVerticalBox::FSlot* AdditionalWidgetSlot = nullptr;
