Tagged Asset Browser
- Allows CreateAssetWindow dialogue to make use of the asset definition for its created class for its display name
- Fixed crash with invalid sharedptr access

diff --git a/Engine/Source/Editor/UserAssetTagsEditor/Private/Widgets/STaggedAssetBrowser.cpp b/Engine/Source/Editor/UserAssetTagsEditor/Private/Widgets/STaggedAssetBrowser.cpp
--- a/Engine/Source/Editor/UserAssetTagsEditor/Private/Widgets/STaggedAssetBrowser.cpp
+++ b/Engine/Source/Editor/UserAssetTagsEditor/Private/Widgets/STaggedAssetBrowser.cpp
@@ -235,9 +235,14 @@ void STaggedAssetBrowser::InitAssetContextMenu()
 			}
 		});
 
-		FToolMenuCanExecuteAction CanBrowseToAssetAction = FToolMenuCanExecuteAction::CreateLambda([this](const FToolMenuContext& Context)
+		FToolMenuCanExecuteAction CanBrowseToAssetAction = FToolMenuCanExecuteAction::CreateLambda([WeakPtr = AsShared().ToWeakPtr()](const FToolMenuContext& Context)
 		{
-			TSharedPtr<SWindow> Window = FSlateApplication::Get().FindWidgetWindow(AsShared());
+			if (WeakPtr.IsValid() == false)
+			{
+				return false;
+			}
+				
+			TSharedPtr<SWindow> Window = FSlateApplication::Get().FindWidgetWindow(WeakPtr.Pin().ToSharedRef());
 			if(Window && Window->IsModalWindow())
 			{
 				return false;
diff --git a/Engine/Source/Editor/UserAssetTagsEditor/Private/Widgets/STaggedAssetBrowserCreateAssetWindow.cpp b/Engine/Source/Editor/UserAssetTagsEditor/Private/Widgets/STaggedAssetBrowserCreateAssetWindow.cpp
--- a/Engine/Source/Editor/UserAssetTagsEditor/Private/Widgets/STaggedAssetBrowserCreateAssetWindow.cpp
+++ b/Engine/Source/Editor/UserAssetTagsEditor/Private/Widgets/STaggedAssetBrowserCreateAssetWindow.cpp
@@ -11,6 +11,16 @@
 void STaggedAssetBrowserCreateAssetWindow::Construct(const FArguments& InArgs, const UTaggedAssetBrowserConfiguration& Configuration, UClass& InCreatedClass)
 {
 	CreatedClass = &InCreatedClass;
+	
+	if(InArgs._bUseAssetDefinition)
+	{
+		AssetDefinition = UAssetDefinitionRegistry::Get()->GetAssetDefinitionForClass(CreatedClass.Get());
+		if(!ensureMsgf(AssetDefinition.IsValid(), TEXT("AssetDefinition should be valid if bUseAssetDefinition is true.")))
+		{
+			return;
+		}
+	}
+	
 	AdditionalFactorySettingsClass = InArgs._AdditionalFactorySettingsClass;
 	bAllowEmptyAssetCreation = InArgs._bAllowEmptyAssetCreation;
 
@@ -41,7 +51,7 @@ void STaggedAssetBrowserCreateAssetWindow::Construct(const FArguments& InArgs, c
 					.AutoWidth()
 					.Padding(2.f)
 					[
-						SNew(STextBlock).Text(FText::FormatOrdered(LOCTEXT("CreateEmptyAssetButtonLabel", "Create Empty {0}"), CreatedClass->GetDisplayNameText()))
+						SNew(STextBlock).Text(FText::FormatOrdered(LOCTEXT("CreateEmptyAssetButtonLabel", "Create Empty {0}"), GetAssetTypeName()))
 					]
 				]
 			]
@@ -125,10 +135,15 @@ FReply STaggedAssetBrowserCreateAssetWindow::Cancel()
 	return FReply::Handled();
 }
 
+FText STaggedAssetBrowserCreateAssetWindow::GetAssetTypeName() const
+{
+	return AssetDefinition.IsValid() ? AssetDefinition->GetAssetDisplayName() : CreatedClass->GetDisplayNameText();
+}
+
 FText STaggedAssetBrowserCreateAssetWindow::GetCreateButtonTooltip() const
 {
 	return HasSelectedAssets()
-	? FText::FormatOrdered(LOCTEXT("CreateAssetButtonTooltip_Enabled", "Create a new {0} with selected asset {1}"), CreatedClass->GetDisplayNameText(), FText::FromName(GetSelectedAssets()[0].AssetName))
+	? FText::FormatOrdered(LOCTEXT("CreateAssetButtonTooltip_Enabled", "Create a new {0} with selected asset {1}"), GetAssetTypeName(), FText::FromName(GetSelectedAssets()[0].AssetName))
 	: FText::FormatOrdered(LOCTEXT("CreateAssetButtonTooltip_Disabled", "Please select an asset as a base for your new effect.{0}"), bAllowEmptyAssetCreation ? LOCTEXT("CreateAssetButtonTooltip_Disabled_AllowsEmptyCreation", " Alternatively, create an empty asset.") : FText::GetEmpty());
 }
 
diff --git a/Engine/Source/Editor/UserAssetTagsEditor/Public/Widgets/STaggedAssetBrowserCreateAssetWindow.h b/Engine/Source/Editor/UserAssetTagsEditor/Public/Widgets/STaggedAssetBrowserCreateAssetWindow.h
--- a/Engine/Source/Editor/UserAssetTagsEditor/Public/Widgets/STaggedAssetBrowserCreateAssetWindow.h
+++ b/Engine/Source/Editor/UserAssetTagsEditor/Public/Widgets/STaggedAssetBrowserCreateAssetWindow.h
@@ -17,12 +17,16 @@ class STaggedAssetBrowserCreateAssetWindow : public STaggedAssetBrowserWindow
 public:
 	SLATE_BEGIN_ARGS(STaggedAssetBrowserCreateAssetWindow)
 		: _AdditionalFactorySettingsClass(nullptr)
+		, _bUseAssetDefinition(true)
 		, _bAllowEmptyAssetCreation(true)
 		{
 		}
 		SLATE_ARGUMENT(STaggedAssetBrowserWindow::FArguments, AssetBrowserWindowArgs)
 		/** An additional factory settings class can be specified for additional factory information. */
 		SLATE_ARGUMENT(UClass*, AdditionalFactorySettingsClass)
+		/** If the asset definition is used, data such as its AssetDisplayName will be used. If not, the class display name is used. */
+		SLATE_ARGUMENT(bool, bUseAssetDefinition)
+		/** If empty asset creation is allowed, a "Create Empty" button will show up on the bottom left that will terminate the window without a selection. */
 		SLATE_ARGUMENT(bool, bAllowEmptyAssetCreation)
 	SLATE_END_ARGS()
 
@@ -42,6 +46,8 @@ private:
 	FReply Proceed();
 	FReply Cancel();
 
+	FText GetAssetTypeName() const;
+	
 	FText GetCreateButtonTooltip() const;
 	
 	TSharedRef<SWidget> CreateFactorySettingsTab();
@@ -51,6 +57,7 @@ private:
 private:
 	TWeakObjectPtr<UClass> CreatedClass;
 	TWeakObjectPtr<UClass> AdditionalFactorySettingsClass;
+	TWeakObjectPtr<const UAssetDefinition> AssetDefinition;
 
 	TObjectPtr<UObject> FactorySettingsObject;
 	TSharedPtr<class IDetailsView> FactorySettingsWidget;
