[Backout] - CL46770095
#rnx
Original CL Desc
-----------------------------------------------------------------
Allow Abstract classes Tagged Asset Browser Filter
Allow Tagged Asset Browser searching for NativeParentClass
Allow Tagged Asset Browser searching for Skeleton

#tests editor

diff --git a/Engine/Plugins/Experimental/UAF/UAF/Source/UAFEditor/Private/Common/TaggedAssetBrowser_UAFBrowserFilters.cpp b/Engine/Plugins/Experimental/UAF/UAF/Source/UAFEditor/Private/Common/TaggedAssetBrowser_UAFBrowserFilters.cpp
deleted file mode 100644
diff --git a/Engine/Plugins/Experimental/UAF/UAF/Source/UAFEditor/Private/Common/TaggedAssetBrowser_UAFBrowserFilters.h b/Engine/Plugins/Experimental/UAF/UAF/Source/UAFEditor/Private/Common/TaggedAssetBrowser_UAFBrowserFilters.h
deleted file mode 100644
diff --git a/Engine/Plugins/Experimental/UAF/UAF/Source/UAFEditor/UAFEditor.Build.cs b/Engine/Plugins/Experimental/UAF/UAF/Source/UAFEditor/UAFEditor.Build.cs
--- a/Engine/Plugins/Experimental/UAF/UAF/Source/UAFEditor/UAFEditor.Build.cs
+++ b/Engine/Plugins/Experimental/UAF/UAF/Source/UAFEditor/UAFEditor.Build.cs
@@ -44,7 +44,6 @@ namespace UnrealBuildTool.Rules
 					"WorkspaceEditor",
 					"ContentBrowser",
 					"ContentBrowserData",
-					"DataHierarchyEditor",
 					"UniversalObjectLocator",
 					"UniversalObjectLocatorEditor",
 					"UserAssetTagsEditor",
diff --git a/Engine/Source/Editor/UserAssetTagsEditor/Private/TaggedAssetBrowserFilters/TaggedAssetBrowser_CommonFilters.cpp b/Engine/Source/Editor/UserAssetTagsEditor/Private/TaggedAssetBrowserFilters/TaggedAssetBrowser_CommonFilters.cpp
--- a/Engine/Source/Editor/UserAssetTagsEditor/Private/TaggedAssetBrowserFilters/TaggedAssetBrowser_CommonFilters.cpp
+++ b/Engine/Source/Editor/UserAssetTagsEditor/Private/TaggedAssetBrowserFilters/TaggedAssetBrowser_CommonFilters.cpp
@@ -358,7 +358,6 @@ void UTaggedAssetBrowserFilter_Class::ModifyARFilterInternal(FARFilter& Filter)
 			Filter.ClassPaths.Add(Class->GetClassPathName());
 		}
 	}
-	Filter.bRecursiveClasses = bAllowChildClasses;
 }
 
 #undef LOCTEXT_NAMESPACE
diff --git a/Engine/Source/Editor/UserAssetTagsEditor/Public/TaggedAssetBrowserFilters/TaggedAssetBrowser_CommonFilters.h b/Engine/Source/Editor/UserAssetTagsEditor/Public/TaggedAssetBrowserFilters/TaggedAssetBrowser_CommonFilters.h
--- a/Engine/Source/Editor/UserAssetTagsEditor/Public/TaggedAssetBrowserFilters/TaggedAssetBrowser_CommonFilters.h
+++ b/Engine/Source/Editor/UserAssetTagsEditor/Public/TaggedAssetBrowserFilters/TaggedAssetBrowser_CommonFilters.h
@@ -172,13 +172,9 @@ public:
 	
 	virtual void ModifyARFilterInternal(FARFilter& Filter) const override;
 
-	UPROPERTY(EditAnywhere, Category="Filter", meta=(AllowAbstract))
+	UPROPERTY(EditAnywhere, Category="Filter")
 	TArray<TObjectPtr<UClass>> Classes;
 
-	/** True if we should allow child classes in our search. */
-	UPROPERTY(EditAnywhere, Category = "Filter")
-	bool bAllowChildClasses = true;
-
 	/** Name for this class filter, by default will use first class in list as name if not set */
 	UPROPERTY(EditAnywhere, Category = "Filter")
 	FName FilterName;
diff --git a/Engine/Source/Editor/UserAssetTagsEditor/Public/TaggedAssetBrowserMenuFilters.h b/Engine/Source/Editor/UserAssetTagsEditor/Public/TaggedAssetBrowserMenuFilters.h
--- a/Engine/Source/Editor/UserAssetTagsEditor/Public/TaggedAssetBrowserMenuFilters.h
+++ b/Engine/Source/Editor/UserAssetTagsEditor/Public/TaggedAssetBrowserMenuFilters.h
@@ -25,7 +25,7 @@ struct FTaggedAssetBrowserContext
  *  While filters are generally part of the hierarchy, sections can also specify additional filters
  */
 UCLASS(Abstract, EditInlineNew)
-class USERASSETTAGSEDITOR_API UTaggedAssetBrowserFilterBase : public UHierarchyItem
+class UTaggedAssetBrowserFilterBase : public UHierarchyItem
 {
 	GENERATED_BODY()
 	
