Allow Abstract classes Tagged Asset Browser Filter
Allow Tagged Asset Browser searching for NativeParentClass
Allow Tagged Asset Browser searching for Skeleton
Fix FObjectInitializer constructor not created due to addition of default ctor.

#was-backout 46772030
#tests editor

diff --git a/Engine/Plugins/Experimental/UAF/UAF/Source/UAFEditor/Private/Common/TaggedAssetBrowser_UAFBrowserFilters.cpp b/Engine/Plugins/Experimental/UAF/UAF/Source/UAFEditor/Private/Common/TaggedAssetBrowser_UAFBrowserFilters.cpp
new file mode 100644
--- /dev/null
+++ b/Engine/Plugins/Experimental/UAF/UAF/Source/UAFEditor/Private/Common/TaggedAssetBrowser_UAFBrowserFilters.cpp
@@ -0,0 +1,87 @@
+// Copyright Epic Games, Inc. All Rights Reserved.
+
+#include "Common/TaggedAssetBrowser_UAFBrowserFilters.h"
+
+#include "Animation/Skeleton.h"
+#include "IContentBrowserDataModule.h"
+#include "UserAssetTagEditorUtilities.h"
+#include "UserAssetTagsEditorModule.h"
+#include "Logging/StructuredLog.h"
+#include "Styling/SlateIconFinder.h"
+#include "Widgets/STaggedAssetBrowser.h"
+
+#include UE_INLINE_GENERATED_CPP_BY_NAME(TaggedAssetBrowser_UAFBrowserFilters)
+
+#define LOCTEXT_NAMESPACE "UserAssetTags"
+
+
+//////////////////////////////////////////////////////////////////////////
+// UTaggedAssetBrowserFilter_NativeClass
+
+void UTaggedAssetBrowserFilter_NativeClass::ModifyARFilterInternal(FARFilter& Filter) const
+{
+	for(UClass* Class : Classes)
+	{
+		if(Class)
+		{
+			if (bSearchNativeParentClasses)
+			{
+				Filter.TagsAndValues.Add(FBlueprintTags::NativeParentClassPath, FObjectPropertyBase::GetExportPath(Class));
+			}
+			else
+			{
+				Filter.ClassPaths.Add(Class->GetClassPathName());
+			}
+		}
+	}
+	Filter.bRecursiveClasses = bAllowChildClasses;
+}
+
+
+//////////////////////////////////////////////////////////////////////////
+// UTaggedAssetBrowserFilter_Skeleton
+
+FString UTaggedAssetBrowserFilter_Skeleton::ToString() const
+{
+	if (!FilterName.IsNone())
+	{
+		return FilterName.ToString();
+	}
+
+	if (Skeleton)
+	{
+		return Skeleton->GetFName().ToString();
+	}
+
+	return UTaggedAssetBrowserFilterBase::ToString();
+}
+
+FText UTaggedAssetBrowserFilter_Skeleton::GetTooltip() const
+{
+	if (Skeleton)
+	{
+		FText BaseTooltip = LOCTEXT("UTaggedAssetBrowserFilter_ClassTooltip", "Skeleton Filter: {0}.");
+		return FText::FormatOrdered(BaseTooltip, FText::FromString(Skeleton->GetFName().ToString()));
+	}
+
+	return UTaggedAssetBrowserFilterBase::GetTooltip();
+}
+
+FSlateIcon UTaggedAssetBrowserFilter_Skeleton::GetIcon() const
+{
+	return FSlateIconFinder::FindIconForClass(USkeleton::StaticClass());
+}
+
+void UTaggedAssetBrowserFilter_Skeleton::ModifyARFilterInternal(FARFilter& Filter) const
+{
+	static const FName SkeletonTagName("Skeleton");
+	if(Skeleton)
+	{
+		Filter.TagsAndValues.Add(SkeletonTagName, FAssetData(Skeleton).GetExportTextName());
+	}
+}
+
+
+//////////////////////////////////////////////////////////////////////////
+
+#undef LOCTEXT_NAMESPACE
diff --git a/Engine/Plugins/Experimental/UAF/UAF/Source/UAFEditor/Private/Common/TaggedAssetBrowser_UAFBrowserFilters.h b/Engine/Plugins/Experimental/UAF/UAF/Source/UAFEditor/Private/Common/TaggedAssetBrowser_UAFBrowserFilters.h
new file mode 100644
--- /dev/null
+++ b/Engine/Plugins/Experimental/UAF/UAF/Source/UAFEditor/Private/Common/TaggedAssetBrowser_UAFBrowserFilters.h
@@ -0,0 +1,56 @@
+// Copyright Epic Games, Inc. All Rights Reserved.
+
+#pragma once
+
+#include "TaggedAssetBrowserFilters/TaggedAssetBrowser_CommonFilters.h"
+#include "TaggedAssetBrowser_UAFBrowserFilters.generated.h"
+
+#define UE_API UAFEDITOR_API
+
+class USkeleton;
+
+/** 
+ * Experimental. Temp filter to allow us to find BP whose BPGC would match a class filter.
+ * 
+ * Our solution here isn't as robust as we would like. But it does work for UAF Browser MVP needs.
+ */
+UCLASS(DisplayName="Native Classes")
+class UE_API UTaggedAssetBrowserFilter_NativeClass : public UTaggedAssetBrowserFilter_Class
+{
+	GENERATED_BODY()
+
+public:
+
+	virtual void ModifyARFilterInternal(FARFilter& Filter) const override;
+
+	/** True to search for a native parent class, useful for finding BPs whose BPGCs would match but BP itself does not match. */
+	UPROPERTY(EditAnywhere, Category = "Filter")
+	bool bSearchNativeParentClasses = false;
+};
+
+
+/** Filter for a specifc skeleton */
+UCLASS(DisplayName="Skeleton")
+class UE_API UTaggedAssetBrowserFilter_Skeleton : public UTaggedAssetBrowserFilterBase
+{
+	GENERATED_BODY()
+
+public:
+
+	virtual FString ToString() const override;
+	virtual TOptional<FString> GetInstanceIdentifier() const override { return ToString(); }
+	virtual FText GetTooltip() const override;
+	virtual FSlateIcon GetIcon() const override;
+
+	virtual void ModifyARFilterInternal(FARFilter& Filter) const override;
+
+	/** Skeleton to search for. */
+	UPROPERTY(EditAnywhere, Category = "Filter")
+	TObjectPtr<USkeleton> Skeleton;
+
+	/** Name for this class filter, by default will use name of the skeleton */
+	UPROPERTY(EditAnywhere, Category = "Filter")
+	FName FilterName;
+};
+
+#undef UE_API
\ No newline at end of file
diff --git a/Engine/Plugins/Experimental/UAF/UAF/Source/UAFEditor/UAFEditor.Build.cs b/Engine/Plugins/Experimental/UAF/UAF/Source/UAFEditor/UAFEditor.Build.cs
--- a/Engine/Plugins/Experimental/UAF/UAF/Source/UAFEditor/UAFEditor.Build.cs
+++ b/Engine/Plugins/Experimental/UAF/UAF/Source/UAFEditor/UAFEditor.Build.cs
@@ -44,6 +44,7 @@ namespace UnrealBuildTool.Rules
 					"WorkspaceEditor",
 					"ContentBrowser",
 					"ContentBrowserData",
+					"DataHierarchyEditor",
 					"UniversalObjectLocator",
 					"UniversalObjectLocatorEditor",
 					"UserAssetTagsEditor",
diff --git a/Engine/Source/Editor/UserAssetTagsEditor/Private/TaggedAssetBrowserFilters/TaggedAssetBrowser_CommonFilters.cpp b/Engine/Source/Editor/UserAssetTagsEditor/Private/TaggedAssetBrowserFilters/TaggedAssetBrowser_CommonFilters.cpp
--- a/Engine/Source/Editor/UserAssetTagsEditor/Private/TaggedAssetBrowserFilters/TaggedAssetBrowser_CommonFilters.cpp
+++ b/Engine/Source/Editor/UserAssetTagsEditor/Private/TaggedAssetBrowserFilters/TaggedAssetBrowser_CommonFilters.cpp
@@ -45,11 +45,6 @@ const FSlateBrush* FTaggedAssetBrowserSectionIconData::GetImageBrush() const
 	return FAppStyle::GetNoBrush();
 }
 
-UTaggedAssetBrowserFilter_All::UTaggedAssetBrowserFilter_All()
-{
-	
-}
-
 FString UTaggedAssetBrowserFilter_UserAssetTag::ToString() const
 {
 	return UserAssetTag.ToString();
@@ -101,11 +96,6 @@ void UTaggedAssetBrowserFilter_UserAssetTag::SetUserAssetTag(FName InUserAssetTa
 	UserAssetTag = InUserAssetTag;
 }
 
-UTaggedAssetBrowserFilter_UserAssetTagCollection::UTaggedAssetBrowserFilter_UserAssetTagCollection()
-{
-	
-}
-
 FString UTaggedAssetBrowserFilter_UserAssetTagCollection::ToString() const
 {
 	return Name.ToString();
@@ -358,6 +348,7 @@ void UTaggedAssetBrowserFilter_Class::ModifyARFilterInternal(FARFilter& Filter)
 			Filter.ClassPaths.Add(Class->GetClassPathName());
 		}
 	}
+	Filter.bRecursiveClasses = bAllowChildClasses;
 }
 
 #undef LOCTEXT_NAMESPACE
diff --git a/Engine/Source/Editor/UserAssetTagsEditor/Public/TaggedAssetBrowserFilters/TaggedAssetBrowser_CommonFilters.h b/Engine/Source/Editor/UserAssetTagsEditor/Public/TaggedAssetBrowserFilters/TaggedAssetBrowser_CommonFilters.h
--- a/Engine/Source/Editor/UserAssetTagsEditor/Public/TaggedAssetBrowserFilters/TaggedAssetBrowser_CommonFilters.h
+++ b/Engine/Source/Editor/UserAssetTagsEditor/Public/TaggedAssetBrowserFilters/TaggedAssetBrowser_CommonFilters.h
@@ -60,8 +60,6 @@ UCLASS(DisplayName="All")
 class USERASSETTAGSEDITOR_API UTaggedAssetBrowserFilter_All : public UTaggedAssetBrowserFilterBase
 {
 	GENERATED_BODY()
-public:
-	UTaggedAssetBrowserFilter_All();
 };
 
 /** A filter that returns assets that have the specified User Asset Tag. Nested User Asset Tag filters will be combined. */
@@ -70,8 +68,6 @@ class USERASSETTAGSEDITOR_API UTaggedAssetBrowserFilter_UserAssetTag : public UT
 {
 	GENERATED_BODY()
 public:
-	UTaggedAssetBrowserFilter_UserAssetTag() = default;
-	
 	virtual FString ToString() const override;
 	virtual TOptional<FString> GetInstanceIdentifier() const override { return ToString(); }
 	virtual FSlateIcon GetIcon() const override;
@@ -91,8 +87,8 @@ UCLASS(DisplayName="User Asset Tag Collection")
 class USERASSETTAGSEDITOR_API UTaggedAssetBrowserFilter_UserAssetTagCollection : public UTaggedAssetBrowserFilterBase
 {
 	GENERATED_BODY()
+
 public:
-	UTaggedAssetBrowserFilter_UserAssetTagCollection();
 
 	virtual FString ToString() const override;
 	virtual TOptional<FString> GetInstanceIdentifier() const override { return ToString(); }
@@ -123,9 +119,9 @@ UCLASS(DisplayName="Recent")
 class USERASSETTAGSEDITOR_API UTaggedAssetBrowserFilter_Recent : public UTaggedAssetBrowserFilterBase
 {
 	GENERATED_BODY()
+
 public:
-	UTaggedAssetBrowserFilter_Recent() = default;
-	
+
 	virtual void InitializeInternal(const FTaggedAssetBrowserContext& InContext) override;
 	virtual bool ShouldFilterAssetInternal(const FAssetData& InAssetData) const override;
 	virtual FSlateIcon GetIcon() const override;
@@ -140,8 +136,8 @@ UCLASS(DisplayName="Directories")
 class USERASSETTAGSEDITOR_API UTaggedAssetBrowserFilter_Directories : public UTaggedAssetBrowserFilterBase
 {
 	GENERATED_BODY()
+
 public:
-	UTaggedAssetBrowserFilter_Directories() = default;
 
 	virtual FString ToString() const override;
 	virtual TOptional<FString> GetInstanceIdentifier() const override { return ToString(); }
@@ -162,8 +158,8 @@ UCLASS(DisplayName="Classes")
 class USERASSETTAGSEDITOR_API UTaggedAssetBrowserFilter_Class : public UTaggedAssetBrowserFilterBase
 {
 	GENERATED_BODY()
+
 public:
-	UTaggedAssetBrowserFilter_Class() = default;
 
 	virtual FString ToString() const override;
 	virtual TOptional<FString> GetInstanceIdentifier() const override { return ToString(); }
@@ -172,9 +168,13 @@ public:
 	
 	virtual void ModifyARFilterInternal(FARFilter& Filter) const override;
 
-	UPROPERTY(EditAnywhere, Category="Filter")
+	UPROPERTY(EditAnywhere, Category="Filter", meta=(AllowAbstract))
 	TArray<TObjectPtr<UClass>> Classes;
 
+	/** True if we should allow child classes in our search. */
+	UPROPERTY(EditAnywhere, Category = "Filter")
+	bool bAllowChildClasses = true;
+
 	/** Name for this class filter, by default will use first class in list as name if not set */
 	UPROPERTY(EditAnywhere, Category = "Filter")
 	FName FilterName;
diff --git a/Engine/Source/Editor/UserAssetTagsEditor/Public/TaggedAssetBrowserMenuFilters.h b/Engine/Source/Editor/UserAssetTagsEditor/Public/TaggedAssetBrowserMenuFilters.h
--- a/Engine/Source/Editor/UserAssetTagsEditor/Public/TaggedAssetBrowserMenuFilters.h
+++ b/Engine/Source/Editor/UserAssetTagsEditor/Public/TaggedAssetBrowserMenuFilters.h
@@ -25,13 +25,11 @@ struct FTaggedAssetBrowserContext
  *  While filters are generally part of the hierarchy, sections can also specify additional filters
  */
 UCLASS(Abstract, EditInlineNew)
-class UTaggedAssetBrowserFilterBase : public UHierarchyItem
+class USERASSETTAGSEDITOR_API UTaggedAssetBrowserFilterBase : public UHierarchyItem
 {
 	GENERATED_BODY()
 	
 public:	
-	UTaggedAssetBrowserFilterBase() {}
-
 	/** The filter has a chance to initialize itself with contextual data to actually work. */
 	void Initialize(const FTaggedAssetBrowserContext& InContext);
 
