Tagged ~700 Niagara Systems based on Yi's FX Browser tag list

#rnx
#changelist validated
#virtualized

diff --git a/Engine/Source/Editor/UserAssetTagsEditor/Public/UserAssetTagMenuHelpers.h b/Engine/Source/Editor/UserAssetTagsEditor/Public/UserAssetTagMenuHelpers.h
--- a/Engine/Source/Editor/UserAssetTagsEditor/Public/UserAssetTagMenuHelpers.h
+++ b/Engine/Source/Editor/UserAssetTagsEditor/Public/UserAssetTagMenuHelpers.h
@@ -30,7 +30,7 @@ namespace UE::UserAssetTags::Menus
 	FMultipleAssetsTagInfo GatherInfoAboutTags(const TArray<FAssetData>& Assets);
 	TOptional<FText> DetermineInfoText(const FMultipleAssetsTagSingleTagInfo& SingleTagInfo, const FMultipleAssetsTagInfo& TotalInfo);
 
-	void AddTagsToAssets(FName InUserAssetTag, const TArray<FAssetData>& InAssetData);
+	USERASSETTAGSEDITOR_API void AddTagsToAssets(FName InUserAssetTag, const TArray<FAssetData>& InAssetData);
 	void RemoveTagsFromAssets(FName InUserAssetTag, const TArray<FAssetData>& InAssetData);
 
 	ECheckBoxState DetermineTagUIStatus(FName InUserAssetTag, const TArray<FAssetData>& InAssetData);
