Added Async Support for Tagged Asset Browser Window
- Minor Tagged Asset Browser Tweaks for cleanliness
- Renamed the CreateAssetWindow to AssetFactoryWindow to clear up intent
- Removed "RequestDestroyWindow" and "OnAssetsActivated" from the parent STaggedAssetBrowserWindow implemention upon asset activation letting the user define this behavior

diff --git a/Engine/Plugins/Editor/MaterialAssetWizard/Source/MaterialAssetWizard/Private/MaterialFactoryWithAssetWizard.cpp b/Engine/Plugins/Editor/MaterialAssetWizard/Source/MaterialAssetWizard/Private/MaterialFactoryWithAssetWizard.cpp
--- a/Engine/Plugins/Editor/MaterialAssetWizard/Source/MaterialAssetWizard/Private/MaterialFactoryWithAssetWizard.cpp
+++ b/Engine/Plugins/Editor/MaterialAssetWizard/Source/MaterialAssetWizard/Private/MaterialFactoryWithAssetWizard.cpp
@@ -8,7 +8,7 @@
 #include "Assets/TaggedAssetBrowserConfiguration.h"
 #include "Interfaces/IMainFrameModule.h"
 #include "Widgets/STaggedAssetBrowser.h"
-#include "Widgets/STaggedAssetBrowserCreateAssetWindow.h"
+#include "Widgets/STaggedAssetBrowserAssetFactoryWindow.h"
 #include "Materials/Material.h"
 #include "Materials/MaterialInstanceConstant.h"
 
@@ -37,7 +37,7 @@ bool UMaterialFactoryWithAssetWizardBase::ConfigureProperties()
 		AssetBrowserWindowArgs.AssetBrowserArgs(AssetBrowserArgs);
 		AssetBrowserWindowArgs.WindowArgs(WindowArgs);
 		
-		TSharedRef<STaggedAssetBrowserCreateAssetWindow> CreateAssetBrowserWindow = SNew(STaggedAssetBrowserCreateAssetWindow, *ConfigurationAsset, *SupportedClass.Get())
+		TSharedRef<STaggedAssetBrowserAssetFactoryWindow> CreateAssetBrowserWindow = SNew(STaggedAssetBrowserAssetFactoryWindow, *ConfigurationAsset, *SupportedClass.Get())
 			.AssetBrowserWindowArgs(AssetBrowserWindowArgs);
 		
 		FSlateApplication::Get().AddModalWindow(CreateAssetBrowserWindow, ParentWindow);
diff --git a/Engine/Plugins/Editor/MaterialAssetWizard/Source/MaterialAssetWizard/Public/MaterialFactoryWithAssetWizard.h b/Engine/Plugins/Editor/MaterialAssetWizard/Source/MaterialAssetWizard/Public/MaterialFactoryWithAssetWizard.h
--- a/Engine/Plugins/Editor/MaterialAssetWizard/Source/MaterialAssetWizard/Public/MaterialFactoryWithAssetWizard.h
+++ b/Engine/Plugins/Editor/MaterialAssetWizard/Source/MaterialAssetWizard/Public/MaterialFactoryWithAssetWizard.h
@@ -7,7 +7,7 @@
 #include "Widgets/STaggedAssetBrowser.h"
 #include "MaterialFactoryWithAssetWizard.generated.h"
 
-class STaggedAssetBrowserCreateAssetWindow;
+class STaggedAssetBrowserAssetFactoryWindow;
 class UMaterial;
 
 /** Abstract base factory that summons material wizard. */
diff --git a/Engine/Plugins/FX/Niagara/Source/NiagaraEditor/Private/NiagaraEmitterFactoryNew.cpp b/Engine/Plugins/FX/Niagara/Source/NiagaraEditor/Private/NiagaraEmitterFactoryNew.cpp
--- a/Engine/Plugins/FX/Niagara/Source/NiagaraEditor/Private/NiagaraEmitterFactoryNew.cpp
+++ b/Engine/Plugins/FX/Niagara/Source/NiagaraEditor/Private/NiagaraEmitterFactoryNew.cpp
@@ -24,7 +24,7 @@
 #include "Interfaces/IMainFrameModule.h"
 #include "Framework/Application/SlateApplication.h"
 #include "Widgets/STaggedAssetBrowser.h"
-#include "Widgets/STaggedAssetBrowserCreateAssetWindow.h"
+#include "Widgets/STaggedAssetBrowserAssetFactoryWindow.h"
 #include "Widgets/AssetBrowser/SNiagaraAddEmitterToSystemWindow.h"
 
 #include UE_INLINE_GENERATED_CPP_BY_NAME(NiagaraEmitterFactoryNew)
@@ -94,7 +94,7 @@ bool UNiagaraEmitterFactoryNew::ConfigureProperties()
 		AssetBrowserWindowArgs.AssetBrowserArgs(AssetBrowserArgs);
 		AssetBrowserWindowArgs.WindowArgs(WindowArgs);
 		
-		TSharedRef<STaggedAssetBrowserCreateAssetWindow> CreateAssetBrowserWindow = SNew(STaggedAssetBrowserCreateAssetWindow, *ConfigurationAsset, *UNiagaraEmitter::StaticClass())
+		TSharedRef<STaggedAssetBrowserAssetFactoryWindow> CreateAssetBrowserWindow = SNew(STaggedAssetBrowserAssetFactoryWindow, *ConfigurationAsset, *UNiagaraEmitter::StaticClass())
 			.AssetBrowserWindowArgs(AssetBrowserWindowArgs);
 		
 		FSlateApplication::Get().AddModalWindow(CreateAssetBrowserWindow, ParentWindow);
@@ -110,10 +110,13 @@ bool UNiagaraEmitterFactoryNew::ConfigureProperties()
 	// Give the minimal emitter a chance to be selected if the user chose 'Minimal Emitter'
 	if(SelectedAssetData.Num() == 0)
 	{
-		FSoftObjectPath MinimalEmitter = GetDefault<UNiagaraEditorSettings>()->DefaultEmptyEmitter;
-		if(MinimalEmitter.IsValid() && MinimalEmitter.IsAsset())
+		FSoftObjectPath MinimalEmitterPath = GetDefault<UNiagaraEditorSettings>()->DefaultEmptyEmitter;
+		if(MinimalEmitterPath.IsValid() && MinimalEmitterPath.IsAsset())
 		{
-			SelectedAssetData.Add(FAssetData(MinimalEmitter.TryLoad()));
+			if(UObject* MinimalEmitter = MinimalEmitterPath.TryLoad())
+			{
+				SelectedAssetData.Add(FAssetData(MinimalEmitter));
+			}
 		}
 	}
 
diff --git a/Engine/Plugins/FX/Niagara/Source/NiagaraEditor/Private/NiagaraSystemFactoryNew.cpp b/Engine/Plugins/FX/Niagara/Source/NiagaraEditor/Private/NiagaraSystemFactoryNew.cpp
--- a/Engine/Plugins/FX/Niagara/Source/NiagaraEditor/Private/NiagaraSystemFactoryNew.cpp
+++ b/Engine/Plugins/FX/Niagara/Source/NiagaraEditor/Private/NiagaraSystemFactoryNew.cpp
@@ -23,7 +23,7 @@
 #include "UserAssetTagEditorUtilities.h"
 #include "Assets/TaggedAssetBrowserConfiguration.h"
 #include "Widgets/STaggedAssetBrowser.h"
-#include "Widgets/STaggedAssetBrowserCreateAssetWindow.h"
+#include "Widgets/STaggedAssetBrowserAssetFactoryWindow.h"
 #include "Widgets/AssetBrowser/SNiagaraAddEmitterToSystemWindow.h"
 
 #include UE_INLINE_GENERATED_CPP_BY_NAME(NiagaraSystemFactoryNew)
@@ -76,8 +76,6 @@ bool UNiagaraSystemFactoryNew::ConfigureProperties()
 			.RecentAndFavoritesList(FNiagaraEditorModule::Get().GetRecentsManager()->GetRecentEmitterAndSystemsList())
 			.OnGetExtraFrontendFilters_UObject(this, &UNiagaraSystemFactoryNew::OnGetExtraFrontendFilters)
 			.OnExtendAddFilterMenu_UObject(this, &UNiagaraSystemFactoryNew::OnExtendAddFilterMenu)
-			// @TODO (ME): This is not ideal. Ideally the factory would have access to creation path and/or property handle (if created from property asset picker -> create menu)
-			// then we wouldn't need to use GetLastDirectory which is relatively unsafe
 			.AdditionalReferencingAssets({ PseudoReferencingAsset });
 
 		SWindow::FArguments WindowArgs;
@@ -91,7 +89,7 @@ bool UNiagaraSystemFactoryNew::ConfigureProperties()
 		AssetBrowserWindowArgs.AssetBrowserArgs(AssetBrowserArgs);
 		AssetBrowserWindowArgs.WindowArgs(WindowArgs);
 		
-		TSharedRef<STaggedAssetBrowserCreateAssetWindow> CreateAssetBrowserWindow = SNew(STaggedAssetBrowserCreateAssetWindow, *ConfigurationAsset, *UNiagaraSystem::StaticClass())
+		TSharedRef<STaggedAssetBrowserAssetFactoryWindow> CreateAssetBrowserWindow = SNew(STaggedAssetBrowserAssetFactoryWindow, *ConfigurationAsset, *UNiagaraSystem::StaticClass())
 			.AssetBrowserWindowArgs(AssetBrowserWindowArgs);
 
 		FSlateApplication::Get().AddModalWindow(CreateAssetBrowserWindow, ParentWindow);
diff --git a/Engine/Plugins/FX/Niagara/Source/NiagaraEditor/Private/Widgets/AssetBrowser/SNiagaraAddEmitterToSystemWindow.cpp b/Engine/Plugins/FX/Niagara/Source/NiagaraEditor/Private/Widgets/AssetBrowser/SNiagaraAddEmitterToSystemWindow.cpp
--- a/Engine/Plugins/FX/Niagara/Source/NiagaraEditor/Private/Widgets/AssetBrowser/SNiagaraAddEmitterToSystemWindow.cpp
+++ b/Engine/Plugins/FX/Niagara/Source/NiagaraEditor/Private/Widgets/AssetBrowser/SNiagaraAddEmitterToSystemWindow.cpp
@@ -98,21 +98,36 @@ bool FFrontendFilter_NiagaraSystemEffectType::PassesFilter(const FContentBrowser
 	return false;
 }
 
-SNiagaraAddEmitterToSystemWindow::~SNiagaraAddEmitterToSystemWindow() = default;
+bool SNiagaraAddEmitterToSystemWindow::CheckValidArguments(const FArguments& InArgs) const
+{
+	if(!ensureMsgf(!InArgs._AssetBrowserWindowArgs._AssetBrowserArgs._OnAssetsActivated.IsBound(), TEXT("OnAssetsActivated should not be bound.")))
+	{
+		return false;
+	}
+
+	return true;
+}
 
 void SNiagaraAddEmitterToSystemWindow::Construct(const FArguments& InArgs, const UTaggedAssetBrowserConfiguration& Configuration, TSharedRef<FNiagaraSystemViewModel> SystemViewModel)
 {
 	WeakSystemViewModel = SystemViewModel;
 
+	if(!CheckValidArguments(InArgs))
+	{
+		SWindow::FArguments WindowArguments;
+		SWindow::Construct(WindowArguments);
+		return;
+	}
+	
 	FDefaultDetailsTabConfiguration DefaultDetailsTabConfiguration;
 	DefaultDetailsTabConfiguration.bUseDefaultDetailsTab = true;
 	DefaultDetailsTabConfiguration.EmptySelectionMessage = LOCTEXT("EmptyEmitterSelectionUserText", "Select an emitter to add it to your system, or add a new minimal emitter");
 	
 	STaggedAssetBrowser::FArguments AssetBrowserArgs;
-	AssetBrowserArgs._AvailableClasses = {UNiagaraEmitter::StaticClass()};
+	AssetBrowserArgs._AvailableClasses = { UNiagaraEmitter::StaticClass() };
 	AssetBrowserArgs._AssetSelectionMode = ESelectionMode::Single;
 	AssetBrowserArgs._DefaultDetailsTabConfiguration = DefaultDetailsTabConfiguration;
-	AssetBrowserArgs._OnAssetsActivated = FOnAssetsActivated::CreateSP(this, &SNiagaraAddEmitterToSystemWindow::OnAssetsActivatedInternal);
+	AssetBrowserArgs._OnAssetsActivated = FOnAssetsActivated::CreateSP(this, &SNiagaraAddEmitterToSystemWindow::OnAssetsActivated);
 	AssetBrowserArgs._RecentAndFavoritesList = FNiagaraEditorModule::Get().GetRecentsManager()->GetRecentEmitterAndSystemsList();
 	AssetBrowserArgs._AdditionalReferencingAssets = { FAssetData(&SystemViewModel->GetSystem()) };
 
@@ -190,7 +205,7 @@ void SNiagaraAddEmitterToSystemWindow::Construct(const FArguments& InArgs, const
 	STaggedAssetBrowserWindow::Construct(TaggedAssetBrowserWindowArgs, Configuration);
 }
 
-void SNiagaraAddEmitterToSystemWindow::OnAssetsActivatedInternal(const TArray<FAssetData>& AssetData, EAssetTypeActivationMethod::Type) const
+void SNiagaraAddEmitterToSystemWindow::OnAssetsActivated(const TArray<FAssetData>& AssetData, EAssetTypeActivationMethod::Type)
 {
 	if(WeakSystemViewModel.IsValid())
 	{
@@ -206,6 +221,8 @@ void SNiagaraAddEmitterToSystemWindow::OnAssetsActivatedInternal(const TArray<FA
 			}
 		}
 	}
+
+	RequestDestroyWindow();
 }
 
 FReply SNiagaraAddEmitterToSystemWindow::AddMinimalEmitter()
diff --git a/Engine/Plugins/FX/Niagara/Source/NiagaraEditor/Public/Widgets/AssetBrowser/SNiagaraAddEmitterToSystemWindow.h b/Engine/Plugins/FX/Niagara/Source/NiagaraEditor/Public/Widgets/AssetBrowser/SNiagaraAddEmitterToSystemWindow.h
--- a/Engine/Plugins/FX/Niagara/Source/NiagaraEditor/Public/Widgets/AssetBrowser/SNiagaraAddEmitterToSystemWindow.h
+++ b/Engine/Plugins/FX/Niagara/Source/NiagaraEditor/Public/Widgets/AssetBrowser/SNiagaraAddEmitterToSystemWindow.h
@@ -52,10 +52,12 @@ public:
 	SLATE_END_ARGS()
 
 	UE_API void Construct(const FArguments& InArgs, const UTaggedAssetBrowserConfiguration& Configuration, TSharedRef<FNiagaraSystemViewModel> SystemViewModel);
-	UE_API virtual ~SNiagaraAddEmitterToSystemWindow() override;
+	
 private:
+	bool CheckValidArguments(const FArguments& InArgs) const;
+	
 	/** The function that will be called by our buttons or by the asset picker itself if double-clicking, hitting enter etc. */
-	void OnAssetsActivatedInternal(const TArray<FAssetData>& AssetData, EAssetTypeActivationMethod::Type) const;
+	void OnAssetsActivated(const TArray<FAssetData>& AssetData, EAssetTypeActivationMethod::Type);
 
 	FReply AddMinimalEmitter();
 	FReply AddSelectedEmitters();
diff --git a/Engine/Source/Editor/UserAssetTagsEditor/Private/Widgets/STaggedAssetBrowser.cpp b/Engine/Source/Editor/UserAssetTagsEditor/Private/Widgets/STaggedAssetBrowser.cpp
--- a/Engine/Source/Editor/UserAssetTagsEditor/Private/Widgets/STaggedAssetBrowser.cpp
+++ b/Engine/Source/Editor/UserAssetTagsEditor/Private/Widgets/STaggedAssetBrowser.cpp
@@ -1211,9 +1211,7 @@ EVisibility STaggedAssetBrowser::FDefaultDetailsTabMembers::OnGetThumbnailVisibi
 }
 
 void STaggedAssetBrowserWindow::Construct(const FArguments& InArgs, const UTaggedAssetBrowserConfiguration& Configuration)
-{
-	OnAssetsActivatedDelegate = InArgs._AssetBrowserArgs._OnAssetsActivated;
-	
+{	
 	SWindow::Construct(InArgs._WindowArgs);
 
 	SAssignNew(AssetBrowser, STaggedAssetBrowser, Configuration)
@@ -1227,7 +1225,7 @@ void STaggedAssetBrowserWindow::Construct(const FArguments& InArgs, const UTagge
 		.OnAdditionalShouldFilterAsset(InArgs._AssetBrowserArgs._OnAdditionalShouldFilterAsset)
 		.AssetSelectionMode(InArgs._AssetBrowserArgs._AssetSelectionMode)
 		.OnAssetSelected(InArgs._AssetBrowserArgs._OnAssetSelected)
-		.OnAssetsActivated(this, &STaggedAssetBrowserWindow::OnAssetsActivated)
+		.OnAssetsActivated(InArgs._AssetBrowserArgs._OnAssetsActivated)
 		.OnGenerateAssetTooltipWidget(InArgs._AssetBrowserArgs._OnGenerateAssetTooltipWidget)
 		.OnIsAssetValidForCustomToolTip(InArgs._AssetBrowserArgs._OnIsAssetValidForCustomToolTip)
 		.AdditionalBottomWidget(InArgs._AssetBrowserArgs._AdditionalBottomWidget)
@@ -1237,14 +1235,6 @@ void STaggedAssetBrowserWindow::Construct(const FArguments& InArgs, const UTagge
 	SetContent(AssetBrowser.ToSharedRef());
 }
 
-STaggedAssetBrowserWindow::STaggedAssetBrowserWindow()
-{
-}
-
-STaggedAssetBrowserWindow::~STaggedAssetBrowserWindow()
-{
-}
-
 bool STaggedAssetBrowserWindow::HasSelectedAssets() const
 {
 	return GetSelectedAssets().Num() > 0;
@@ -1255,14 +1245,6 @@ TArray<FAssetData> STaggedAssetBrowserWindow::GetSelectedAssets() const
 	return AssetBrowser->GetSelectedAssets();
 }
 
-void STaggedAssetBrowserWindow::OnAssetsActivated(const TArray<FAssetData>& AssetData, EAssetTypeActivationMethod::Type ActivationType)
-{
-	if(OnAssetsActivatedDelegate.ExecuteIfBound(AssetData, ActivationType))
-	{
-		RequestDestroyWindow();
-	}
-}
-
 void STaggedAssetBrowserWindow::OnFocusLost(const FFocusEvent& InFocusEvent)
 {
 	SWindow::OnFocusLost(InFocusEvent);
diff --git a/Engine/Source/Editor/UserAssetTagsEditor/Private/Widgets/STaggedAssetBrowserAssetFactoryWindow.cpp b/Engine/Source/Editor/UserAssetTagsEditor/Private/Widgets/STaggedAssetBrowserAssetFactoryWindow.cpp
new file mode 100644
--- /dev/null
+++ b/Engine/Source/Editor/UserAssetTagsEditor/Private/Widgets/STaggedAssetBrowserAssetFactoryWindow.cpp
@@ -0,0 +1,220 @@
+﻿// Copyright Epic Games, Inc. All Rights Reserved.
+
+#include "Widgets/STaggedAssetBrowserAssetFactoryWindow.h"
+
+#include "ContentBrowserModule.h"
+#include "EditorDirectories.h"
+#include "Widgets/Layout/SSpacer.h"
+
+#define LOCTEXT_NAMESPACE "TaggedAssetBrowser"
+
+void STaggedAssetBrowserAssetFactoryWindow::Construct(const FArguments& InArgs, const UTaggedAssetBrowserConfiguration& Configuration, UClass& InCreatedClass)
+{
+	CreatedClass = &InCreatedClass;
+
+	if(!CheckValidArguments(InArgs))
+	{
+		SWindow::FArguments WindowArguments;
+		SWindow::Construct(WindowArguments);
+		return;
+	}
+
+	if(InArgs._bUseAssetDefinition)
+	{
+		AssetDefinition = UAssetDefinitionRegistry::Get()->GetAssetDefinitionForClass(CreatedClass.Get());
+	}
+	
+	AsyncFactoryArguments = InArgs._AsyncFactoryArguments;
+	AdditionalFactorySettingsClass = InArgs._AdditionalFactorySettingsClass;
+	bAllowEmptyAssetCreation = InArgs._bAllowEmptyAssetCreation;
+
+	FArguments Args = InArgs;
+
+	TSharedRef<SWidget> CreateAssetControls = SNew(SBox)
+		.Padding(16.f, 16.f)
+		[
+			SNew(SHorizontalBox)
+			+ SHorizontalBox::Slot()
+			.AutoWidth()
+			.Padding(5.f, 1.f)
+			.HAlign(HAlign_Left)
+			[
+				SNew(SButton)
+				.OnClicked(this, &STaggedAssetBrowserAssetFactoryWindow::Proceed)
+				.Visibility(bAllowEmptyAssetCreation ? EVisibility::Visible : EVisibility::Hidden)
+				[
+					SNew(SHorizontalBox)
+					+ SHorizontalBox::Slot()
+					.AutoWidth()
+					.Padding(2.f)
+					[
+						SNew(SImage)
+						.Image(FSlateIconFinder::FindIconForClass(CreatedClass.Get()).GetIcon())
+					]
+					+ SHorizontalBox::Slot()
+					.AutoWidth()
+					.Padding(2.f)
+					[
+						SNew(STextBlock).Text(FText::FormatOrdered(LOCTEXT("CreateEmptyAssetButtonLabel", "Create Empty {0}"), GetAssetTypeName()))
+					]
+				]
+			]
+			+ SHorizontalBox::Slot()
+			[
+				SNew(SSpacer)
+			]
+			+ SHorizontalBox::Slot()
+			.AutoWidth()
+			.Padding(5.f, 1.f)
+			[
+				SNew(SButton)
+				.ContentPadding(FMargin(50.f, 4.f, 50.f, 1.f))
+				.ButtonStyle(&FAppStyle::Get().GetWidgetStyle<FButtonStyle>("PrimaryButton"))
+				.TextStyle(&FAppStyle::Get().GetWidgetStyle<FTextBlockStyle>("PrimaryButtonText"))
+				.HAlign(HAlign_Center)
+				.Text(LOCTEXT("CreatePrimaryButtonLabel", "Create"))
+				.OnClicked(this, &STaggedAssetBrowserAssetFactoryWindow::Proceed)
+				.IsEnabled(this, &STaggedAssetBrowserAssetFactoryWindow::HasSelectedAssets)
+				.ToolTipText(this, &STaggedAssetBrowserAssetFactoryWindow::GetCreateButtonTooltip)
+			]
+			+ SHorizontalBox::Slot()
+			.AutoWidth()
+			.Padding(5.f, 1.f)
+			[
+				SNew(SButton)
+				.ContentPadding(FMargin(50.f, 4.f, 50.f, 1.f))
+				.TextStyle(&FAppStyle::Get().GetWidgetStyle<FTextBlockStyle>("PrimaryButtonText"))
+				.HAlign(HAlign_Center)
+				.Text(LOCTEXT("CancelButtonLabel", "Cancel"))
+				.OnClicked(this, &STaggedAssetBrowserAssetFactoryWindow::Cancel)
+			]
+		];
+
+	if(AdditionalFactorySettingsClass.IsValid())
+	{
+		FTaggedAssetBrowserCustomTabInfo FactoryTabInfo;
+		FactoryTabInfo.Title = LOCTEXT("FactoryTabTitle", "Settings");
+		FactoryTabInfo.Icon = FAppStyle::GetBrush("Icons.Settings");
+		FactoryTabInfo.OnGetTabContent = FOnGetContent::CreateSP(this, &STaggedAssetBrowserAssetFactoryWindow::CreateFactorySettingsTab);
+
+		Args._AssetBrowserWindowArgs._AssetBrowserArgs._CustomTabInfos = { FactoryTabInfo };
+	}
+	
+	Args._AssetBrowserWindowArgs._AssetBrowserArgs._AdditionalBottomWidget = CreateAssetControls;
+	Args._AssetBrowserWindowArgs._AssetBrowserArgs._OnAssetsActivated = FOnAssetsActivated::CreateSP(this, &STaggedAssetBrowserAssetFactoryWindow::OnAssetsActivated);
+
+	// If we haven't specified any referencing assets, we assume the context is the last "save asset" editor path, which should be set before this is called
+	if(Args._AssetBrowserWindowArgs._AssetBrowserArgs._AdditionalReferencingAssets.Num() == 0)
+	{
+		FString PackagePath = FEditorDirectories::Get().GetLastDirectory(ELastDirectory::NEW_ASSET);
+		FString AssetName = "TmpAsset";
+		FName PackageName = FName(PackagePath + "/" + AssetName);
+		FAssetData PseudoReferencingAsset(PackageName, FName(PackagePath), FName(AssetName), UObject::StaticClass()->GetClassPathName());
+		Args._AssetBrowserWindowArgs._AssetBrowserArgs._AdditionalReferencingAssets = { PseudoReferencingAsset };  
+	}
+	
+	STaggedAssetBrowserWindow::Construct(Args._AssetBrowserWindowArgs, Configuration);
+}
+
+bool STaggedAssetBrowserAssetFactoryWindow::CheckValidArguments(const FArguments& InArgs) const
+{
+	if(!ensureMsgf(!InArgs._AssetBrowserWindowArgs._AssetBrowserArgs._OnAssetsActivated.IsBound(), TEXT("OnAssetsActivated should not be bound.")))
+	{
+		return false;
+	}
+
+	if(InArgs._bUseAssetDefinition)
+	{		
+		if(!ensureMsgf(UAssetDefinitionRegistry::Get()->GetAssetDefinitionForClass(CreatedClass.Get()) != nullptr, TEXT("AssetDefinition should be valid if bUseAssetDefinition is true.")))
+		{
+			return false;
+		}
+	}
+
+	return true;
+}
+
+void STaggedAssetBrowserAssetFactoryWindow::OnAssetsActivated(const TArray<FAssetData>& AssetData, EAssetTypeActivationMethod::Type ActivationMethod)
+{
+	// Asset activation will close this window, and resume the factory path to call UFactory::FactoryCreateNew.
+	
+	// If modal, the factory can check this bool for proceeding/canceling right within the ConfigureProperties function after the window closes.
+	bProceedWithAction = true;
+
+	// If non-modal, this will allow the factory to configure its properties based on asset activation.
+	if(AsyncFactoryArguments.IsSet())
+	{
+		AsyncFactoryArguments->OnAssetsActivated.ExecuteIfBound(AssetData, ActivationMethod);
+		AsyncFactoryArguments->OnFactoryConfigurePropertiesComplete.ExecuteIfBound(AsyncFactoryArguments->Factory.Get());
+	}
+	
+	RequestDestroyWindow();
+}
+
+FReply STaggedAssetBrowserAssetFactoryWindow::Proceed()
+{
+	OnAssetsActivated(GetSelectedAssets(), EAssetTypeActivationMethod::Opened);
+	return FReply::Handled();
+}
+
+FReply STaggedAssetBrowserAssetFactoryWindow::Cancel()
+{
+	bProceedWithAction = false;
+
+	if(AsyncFactoryArguments.IsSet())
+	{
+		AsyncFactoryArguments->OnFactoryConfigurePropertiesCancelled.ExecuteIfBound(AsyncFactoryArguments->Factory.Get());
+	}
+	
+	RequestDestroyWindow();
+	return FReply::Handled();
+}
+
+FText STaggedAssetBrowserAssetFactoryWindow::GetAssetTypeName() const
+{
+	return AssetDefinition.IsValid() ? AssetDefinition->GetAssetDisplayName() : CreatedClass->GetDisplayNameText();
+}
+
+FText STaggedAssetBrowserAssetFactoryWindow::GetCreateButtonTooltip() const
+{
+	return HasSelectedAssets()
+	? FText::FormatOrdered(LOCTEXT("CreateAssetButtonTooltip_Enabled", "Create a new {0} with selected asset {1}"), GetAssetTypeName(), FText::FromName(GetSelectedAssets()[0].AssetName))
+	: FText::FormatOrdered(LOCTEXT("CreateAssetButtonTooltip_Disabled", "Please select an asset as a base for your new effect.{0}"), bAllowEmptyAssetCreation ? LOCTEXT("CreateAssetButtonTooltip_Disabled_AllowsEmptyCreation", " Alternatively, create an empty asset.") : FText::GetEmpty());
+}
+
+TSharedRef<SWidget> STaggedAssetBrowserAssetFactoryWindow::CreateFactorySettingsTab()
+{
+	if(AdditionalFactorySettingsClass.IsValid())
+	{
+		FactorySettingsObject = NewObject<UObject>(GetTransientPackage(), AdditionalFactorySettingsClass.Get(), NAME_None, RF_Transient);
+		FPropertyEditorModule& PropertyEditorModule = FModuleManager::Get().LoadModuleChecked<FPropertyEditorModule>("PropertyEditor");
+
+		FDetailsViewArgs Args;
+		Args.bShowObjectLabel = false;
+		Args.NameAreaSettings = FDetailsViewArgs::HideNameArea;
+
+		FactorySettingsWidget = PropertyEditorModule.CreateDetailView(Args);
+		FactorySettingsWidget->SetObject(FactorySettingsObject);
+		
+		return FactorySettingsWidget.ToSharedRef();
+	}
+
+	return SNullWidget::NullWidget;
+}
+
+FString STaggedAssetBrowserAssetFactoryWindow::GetReferencerName() const
+{
+	return TEXT("STaggedAssetBrowserCreateAsset");
+}
+
+void STaggedAssetBrowserAssetFactoryWindow::AddReferencedObjects(FReferenceCollector& Collector)
+{
+	FSlateInvalidationRoot::AddReferencedObjects(Collector);
+	
+	if(FactorySettingsObject)
+	{
+		Collector.AddReferencedObject(FactorySettingsObject);
+	}
+}
+
+#undef LOCTEXT_NAMESPACE
diff --git a/Engine/Source/Editor/UserAssetTagsEditor/Private/Widgets/STaggedAssetBrowserCreateAssetWindow.cpp b/Engine/Source/Editor/UserAssetTagsEditor/Private/Widgets/STaggedAssetBrowserCreateAssetWindow.cpp
deleted file mode 100644
diff --git a/Engine/Source/Editor/UserAssetTagsEditor/Public/Widgets/STaggedAssetBrowser.h b/Engine/Source/Editor/UserAssetTagsEditor/Public/Widgets/STaggedAssetBrowser.h
--- a/Engine/Source/Editor/UserAssetTagsEditor/Public/Widgets/STaggedAssetBrowser.h
+++ b/Engine/Source/Editor/UserAssetTagsEditor/Public/Widgets/STaggedAssetBrowser.h
@@ -105,8 +105,13 @@ public:
 		/** If supplied, determines whether an asset should get a custom tooltip. Requires OnGenerateAssetTooltipWidget to be bound. */
 		SLATE_EVENT(FOnIsAssetValidForCustomToolTip, OnIsAssetValidForCustomToolTip)
 
+		/** Determines the asset selection mode; single asset selection by default. */
 		SLATE_ARGUMENT(ESelectionMode::Type, AssetSelectionMode)
+
+		/** The event that triggers when an asset has been selected. */
 		SLATE_EVENT(FOnAssetSelected, OnAssetSelected)
+
+		/** The event that triggers when an asset has been activated (double clicked, custom buttons to proceed etc.). */
 		SLATE_EVENT(FOnAssetsActivated, OnAssetsActivated)
 		
 		/** An additional widget to attach below the asset browser. */
@@ -304,21 +309,20 @@ public:
 
 	UE_API void Construct(const FArguments& InArgs, const UTaggedAssetBrowserConfiguration& Configuration);
 
-	UE_API STaggedAssetBrowserWindow();
-	UE_API virtual ~STaggedAssetBrowserWindow() override;
+	STaggedAssetBrowserWindow() {}
+	virtual ~STaggedAssetBrowserWindow() override {}
 	
 	UE_API bool HasSelectedAssets() const;
+	
 	UE_API TArray<FAssetData> GetSelectedAssets() const;
 	
-protected:	
-	UE_API void OnAssetsActivated(const TArray<FAssetData>& AssetData, EAssetTypeActivationMethod::Type ActivationType);
-
+protected:
 	UE_API virtual void OnFocusLost(const FFocusEvent& InFocusEvent) override;
+	
 	UE_API virtual FReply OnKeyDown(const FGeometry& MyGeometry, const FKeyEvent& InKeyEvent) override;
+	
 protected:
 	TSharedPtr<STaggedAssetBrowser> AssetBrowser;
-private:
-	FOnAssetsActivated OnAssetsActivatedDelegate;
 };
 
 #undef UE_API
diff --git a/Engine/Source/Editor/UserAssetTagsEditor/Public/Widgets/STaggedAssetBrowserAssetFactoryWindow.h b/Engine/Source/Editor/UserAssetTagsEditor/Public/Widgets/STaggedAssetBrowserAssetFactoryWindow.h
new file mode 100644
--- /dev/null
+++ b/Engine/Source/Editor/UserAssetTagsEditor/Public/Widgets/STaggedAssetBrowserAssetFactoryWindow.h
@@ -0,0 +1,83 @@
+﻿// Copyright Epic Games, Inc. All Rights Reserved.
+
+#pragma once
+
+#include "STaggedAssetBrowser.h"
+#include "Factories/Factory.h"
+
+#define UE_API USERASSETTAGSEDITOR_API
+
+/** This window is meant to be used within a factory.
+ *  Modal: Summon this window from ConfigureProperties, then check "ShouldProceedWithAction", retrieve selected assets and settings.
+ *  Non-Modal: Summon this window from ConfigurePropertiesAsync and pass in the AsyncFactoryArguments.
+ *  In both cases, the actual asset creation invocation of the factory is handled by the Content Browser or the Asset Tools.
+ */
+class STaggedAssetBrowserAssetFactoryWindow : public STaggedAssetBrowserWindow
+{
+public:
+	struct FAsyncFactoryArguments
+	{
+		/** The factory that summoned this window from ConfigurePropertiesAsync. */
+		TStrongObjectPtr<UFactory> Factory;
+		/** A callback to let your factory be informed of asset activation so it can configure its properties accordingly. */
+		FOnAssetsActivated OnAssetsActivated;
+		/** The callbacks from ConfigurePropertiesAsync to let the process know to continue or to cancel. */
+		FOnFactoryConfigurePropertiesAsyncComplete OnFactoryConfigurePropertiesComplete;
+		FOnFactoryConfigurePropertiesAsyncCancelled OnFactoryConfigurePropertiesCancelled;
+	};
+	
+	SLATE_BEGIN_ARGS(STaggedAssetBrowserAssetFactoryWindow)	{}
+		SLATE_ARGUMENT(STaggedAssetBrowserWindow::FArguments, AssetBrowserWindowArgs)
+		/** Async callbacks for use within an async factory */
+		SLATE_ARGUMENT(TOptional<FAsyncFactoryArguments>, AsyncFactoryArguments)
+		/** An additional factory settings class can be specified for additional factory information. */
+		SLATE_ARGUMENT_DEFAULT(UClass*, AdditionalFactorySettingsClass) = nullptr;
+		/** If the asset definition is used, data such as its AssetDisplayName will be used. If not, the class display name is used. */
+		SLATE_ARGUMENT_DEFAULT(bool, bUseAssetDefinition) = true;
+		/** If empty asset creation is allowed, a "Create Empty" button will show up on the bottom left that will terminate the window without a selection. */
+		SLATE_ARGUMENT_DEFAULT(bool, bAllowEmptyAssetCreation) = true;
+	SLATE_END_ARGS()
+
+	UE_API void Construct(const FArguments& InArgs, const UTaggedAssetBrowserConfiguration& Configuration, UClass& CreatedClass);
+
+	STaggedAssetBrowserAssetFactoryWindow() {}
+	virtual ~STaggedAssetBrowserAssetFactoryWindow() override {}
+
+	/** For use in a modal context: retrieve whether the user activated assets or clicked a Proceed button, or closed the window.*/
+	bool ShouldProceedWithAction() const { return bProceedWithAction; }
+
+	/** Retrieve the factory settings object, if available. This can be used to further customize the factory process. */
+	UObject* RetrieveFactorySettings() const { return FactorySettingsObject; }
+	
+private:
+	bool CheckValidArguments(const FArguments& InArgs) const;
+	
+	/** The function that will be called by our buttons or by the asset picker itself if double-clicking, hitting enter etc. */
+	void OnAssetsActivated(const TArray<FAssetData>& AssetData, EAssetTypeActivationMethod::Type);
+	
+	FReply Proceed();
+	FReply Cancel();
+
+	FText GetAssetTypeName() const;
+	
+	FText GetCreateButtonTooltip() const;
+	
+	TSharedRef<SWidget> CreateFactorySettingsTab();
+
+	UE_API virtual FString GetReferencerName() const override;
+	UE_API virtual void AddReferencedObjects(FReferenceCollector& Collector) override;
+private:
+	TWeakObjectPtr<UClass> CreatedClass;
+	TWeakObjectPtr<UClass> AdditionalFactorySettingsClass;
+	TWeakObjectPtr<const UAssetDefinition> AssetDefinition;
+
+	TObjectPtr<UObject> FactorySettingsObject;
+	TSharedPtr<class IDetailsView> FactorySettingsWidget;
+	
+	bool bProceedWithAction = false;
+	bool bAllowEmptyAssetCreation = false;
+
+	TOptional<FAsyncFactoryArguments> AsyncFactoryArguments;
+};
+
+#undef UE_API
diff --git a/Engine/Source/Editor/UserAssetTagsEditor/Public/Widgets/STaggedAssetBrowserCreateAssetWindow.h b/Engine/Source/Editor/UserAssetTagsEditor/Public/Widgets/STaggedAssetBrowserCreateAssetWindow.h
deleted file mode 100644
