Removed tag documentation link while not available

#Synced-CL 47950627

diff --git a/Engine/Source/Editor/UserAssetTagsEditor/Private/Widgets/SUserAssetTagsEditor.cpp b/Engine/Source/Editor/UserAssetTagsEditor/Private/Widgets/SUserAssetTagsEditor.cpp
--- a/Engine/Source/Editor/UserAssetTagsEditor/Private/Widgets/SUserAssetTagsEditor.cpp
+++ b/Engine/Source/Editor/UserAssetTagsEditor/Private/Widgets/SUserAssetTagsEditor.cpp
@@ -366,8 +366,9 @@ TSharedPtr<SWidget> SUserAssetTagsEditor::CreateToolbar()
 					.Image(FAppStyle::Get().GetBrush("Icons.Documentation"))
 				];
 
-				FToolMenuEntry DocumentationMenuEntry = FToolMenuEntry::InitWidget("Documentation", DocumentationWidget, FText::GetEmpty(), true);
-				DynamicMenu->AddMenuEntry(NAME_None, DocumentationMenuEntry);
+				// @TODO (ME): Readd documentation when we have a link
+				// FToolMenuEntry DocumentationMenuEntry = FToolMenuEntry::InitWidget("Documentation", DocumentationWidget, FText::GetEmpty(), true);
+				// DynamicMenu->AddMenuEntry(NAME_None, DocumentationMenuEntry);
 			}
 				
 			// View Options
