Fixup API macros

#rnx

#Synced-CL 48642164

#Synced-CL 48757204

diff --git a/Engine/Plugins/AudioInsights/Source/AudioInsights/Public/Cache/AudioCachedMessageChunk.h b/Engine/Plugins/AudioInsights/Source/AudioInsights/Public/Cache/AudioCachedMessageChunk.h
--- a/Engine/Plugins/AudioInsights/Source/AudioInsights/Public/Cache/AudioCachedMessageChunk.h
+++ b/Engine/Plugins/AudioInsights/Source/AudioInsights/Public/Cache/AudioCachedMessageChunk.h
@@ -15,7 +15,7 @@ namespace UE::Audio::Insights
 	public:
 		using TCachePagedArray = TraceServices::TPagedArray<TSharedPtr<class IAudioCachedMessage>>;
 
-		UE_API FAudioCachedMessageChunk(TraceServices::FSlabAllocator& InAllocator, const uint32 InChunkSize)
+		FAudioCachedMessageChunk(TraceServices::FSlabAllocator& InAllocator, const uint32 InChunkSize)
 			: Allocator(InAllocator)
 			, MaxChunkSize(InChunkSize)
 		{ }
@@ -81,4 +81,4 @@ namespace UE::Audio::Insights
 	};
 } // namespace UE::Audio::Insights
 
-#undef UE_API
\ No newline at end of file
+#undef UE_API
diff --git a/Engine/Plugins/AudioInsights/Source/AudioInsights/Public/Providers/AssetProvider.h b/Engine/Plugins/AudioInsights/Source/AudioInsights/Public/Providers/AssetProvider.h
--- a/Engine/Plugins/AudioInsights/Source/AudioInsights/Public/Providers/AssetProvider.h
+++ b/Engine/Plugins/AudioInsights/Source/AudioInsights/Public/Providers/AssetProvider.h
@@ -18,7 +18,7 @@ namespace UE::Audio::Insights
 	class TAssetProvider
 	{
 	public:
-		UE_API TAssetProvider()
+		TAssetProvider()
 		{
 			FAssetRegistryModule& AssetRegistryModule = FModuleManager::LoadModuleChecked<FAssetRegistryModule>("AssetRegistry");
 
@@ -32,7 +32,7 @@ namespace UE::Audio::Insights
 			FTraceAuxiliary::OnTraceStarted.AddRaw(this, &TAssetProvider::HandleOnTraceStarted);
 		}
 
-		UE_API virtual ~TAssetProvider()
+		virtual ~TAssetProvider()
 		{
 			if (FAssetRegistryModule* AssetRegistryModule = FModuleManager::GetModulePtr<FAssetRegistryModule>("AssetRegistry"))
 			{
@@ -47,7 +47,7 @@ namespace UE::Audio::Insights
 			FTraceAuxiliary::OnTraceStarted.RemoveAll(this);
 		}
 
-		UE_API void RequestEntriesUpdate()
+		void RequestEntriesUpdate()
 		{
 			UpdateAssetNames();
 		}
diff --git a/Engine/Plugins/AudioInsights/Source/AudioInsights/Public/Settings/VisibleColumnsSettings.h b/Engine/Plugins/AudioInsights/Source/AudioInsights/Public/Settings/VisibleColumnsSettings.h
--- a/Engine/Plugins/AudioInsights/Source/AudioInsights/Public/Settings/VisibleColumnsSettings.h
+++ b/Engine/Plugins/AudioInsights/Source/AudioInsights/Public/Settings/VisibleColumnsSettings.h
@@ -28,7 +28,7 @@ struct FVisibleColumnsSettings
 
 	// Override in derived class to enable finding properties by name in the derived USTRUCT
 	// E.g. return StaticStruct()->FindPropertyByName(PropertyName);
-	UE_API virtual const FProperty* FindProperty(const FName& PropertyName) const { return nullptr; }
+	virtual const FProperty* FindProperty(const FName& PropertyName) const { return nullptr; }
 };
 
-#undef UE_API
\ No newline at end of file
+#undef UE_API
diff --git a/Engine/Plugins/Chooser/Source/ChooserEditor/Public/ChooserInitializer.h b/Engine/Plugins/Chooser/Source/ChooserEditor/Public/ChooserInitializer.h
--- a/Engine/Plugins/Chooser/Source/ChooserEditor/Public/ChooserInitializer.h
+++ b/Engine/Plugins/Chooser/Source/ChooserEditor/Public/ChooserInitializer.h
@@ -8,7 +8,7 @@
 #include "ChooserInitializer.generated.h"
 
 USTRUCT()
-struct CHOOSEREDITOR_API FChooserInitializer
+struct FChooserInitializer
 {
 	GENERATED_BODY()
 	virtual ~FChooserInitializer() = default;
diff --git a/Engine/Plugins/Experimental/Dataflow/Source/DataflowEditor/Public/Dataflow/DataflowConstructionViewportClient.h b/Engine/Plugins/Experimental/Dataflow/Source/DataflowEditor/Public/Dataflow/DataflowConstructionViewportClient.h
--- a/Engine/Plugins/Experimental/Dataflow/Source/DataflowEditor/Public/Dataflow/DataflowConstructionViewportClient.h
+++ b/Engine/Plugins/Experimental/Dataflow/Source/DataflowEditor/Public/Dataflow/DataflowConstructionViewportClient.h
@@ -49,7 +49,7 @@ public:
 	UE_API FString GetOverlayString() const;
 
 	/** Return the widget mode only used for tooling */
-	UE_API UE::Widget::EWidgetMode GetToolWidgetMode() const {return WidgetMode;}
+	UE::Widget::EWidgetMode GetToolWidgetMode() const {return WidgetMode;}
 
 protected:
 
diff --git a/Engine/Plugins/Experimental/Dataflow/Source/DataflowEditor/Public/DataflowEditorTools/DataflowEditorSkinWeightsPaintTool.h b/Engine/Plugins/Experimental/Dataflow/Source/DataflowEditor/Public/DataflowEditorTools/DataflowEditorSkinWeightsPaintTool.h
--- a/Engine/Plugins/Experimental/Dataflow/Source/DataflowEditor/Public/DataflowEditorTools/DataflowEditorSkinWeightsPaintTool.h
+++ b/Engine/Plugins/Experimental/Dataflow/Source/DataflowEditor/Public/DataflowEditorTools/DataflowEditorSkinWeightsPaintTool.h
@@ -61,7 +61,7 @@ private:
 	virtual FEditorViewportClient* GetViewportClient() const override {return nullptr;}
 	UE_API virtual void Setup() override;
 	UE_API virtual void OnShutdown(EToolShutdownType ShutdownType) override;
-	UE_API virtual bool CanAccept() const override { return bWeightsHaveChanged; }
+	virtual bool CanAccept() const override { return bWeightsHaveChanged; }
 	UE_API virtual void ToggleBoneManipulation(bool bEnable);
 	UE_API virtual const TArray<FTransform>& GetComponentSpaceBoneTransforms();
 	//~ End USkinWeightsPaintTool interface
diff --git a/Engine/Plugins/Experimental/EditorDataStorageFeatures/Source/TedsOutliner/Public/TedsOutlinerFilter.h b/Engine/Plugins/Experimental/EditorDataStorageFeatures/Source/TedsOutliner/Public/TedsOutlinerFilter.h
--- a/Engine/Plugins/Experimental/EditorDataStorageFeatures/Source/TedsOutliner/Public/TedsOutlinerFilter.h
+++ b/Engine/Plugins/Experimental/EditorDataStorageFeatures/Source/TedsOutliner/Public/TedsOutlinerFilter.h
@@ -15,17 +15,17 @@ namespace UE::Editor::Outliner
 	using namespace UE::Editor::DataStorage::Queries;
 	class FTedsOutlinerImpl;
 
-	struct TEDSOUTLINER_API FTedsFilterData
+	struct FTedsFilterData
 	{
-		FTedsFilterData(const FName& InFilterName, const FText& InFilterDisplayName, const FText& InFilterToolTip, const FName& InFilterIconName,
+		TEDSOUTLINER_API FTedsFilterData(const FName& InFilterName, const FText& InFilterDisplayName, const FText& InFilterToolTip, const FName& InFilterIconName,
 			const TSharedPtr<FFilterCategory>& InCategory, const QueryHandle& InFilterQuery);
 
-		FTedsFilterData(const FName& InFilterName, const FText& InFilterDisplayName, const QueryHandle& InFilterQuery);
+		TEDSOUTLINER_API FTedsFilterData(const FName& InFilterName, const FText& InFilterDisplayName, const QueryHandle& InFilterQuery);
 		
-		FTedsFilterData(const FName& InFilterName, const FText& InFilterDisplayName, const FText& InFilterToolTip, const FName& InFilterIconName,
+		TEDSOUTLINER_API FTedsFilterData(const FName& InFilterName, const FText& InFilterDisplayName, const FText& InFilterToolTip, const FName& InFilterIconName,
 			const TSharedPtr<FFilterCategory>& InCategory, const TQueryFunction<bool>& InFilterQuery);
 
-		FTedsFilterData(const FName& InFilterName, const FText& InFilterDisplayName, const TQueryFunction<bool>& InFilterQuery);
+		TEDSOUTLINER_API FTedsFilterData(const FName& InFilterName, const FText& InFilterDisplayName, const TQueryFunction<bool>& InFilterQuery);
 
 		/*
 		 * Require that the template type cannot be one of the pre-defined FilterQuery types so it will pass to the defined
@@ -62,8 +62,8 @@ namespace UE::Editor::Outliner
 				Queries::BuildQueryFunction<bool>(Forward<FilterFunction>(InFilterQuery)))
 		{}
 
-		FTedsFilterData(const FTedsFilterData&);
-		~FTedsFilterData(); 
+		TEDSOUTLINER_API FTedsFilterData(const FTedsFilterData&);
+		TEDSOUTLINER_API ~FTedsFilterData(); 
 		
 		FName FilterName;
 		FText FilterDisplayName;
diff --git a/Engine/Plugins/Experimental/EditorDataStorageFeatures/Source/TedsTableViewer/Private/QueryStack/TopLevelRowsNode.h b/Engine/Plugins/Experimental/EditorDataStorageFeatures/Source/TedsTableViewer/Private/QueryStack/TopLevelRowsNode.h
--- a/Engine/Plugins/Experimental/EditorDataStorageFeatures/Source/TedsTableViewer/Private/QueryStack/TopLevelRowsNode.h
+++ b/Engine/Plugins/Experimental/EditorDataStorageFeatures/Source/TedsTableViewer/Private/QueryStack/TopLevelRowsNode.h
@@ -20,7 +20,7 @@ namespace UE::Editor::DataStorage::QueryStack
 	public:
 		
 		UE_API FTopLevelRowsNode(const ICoreProvider* InStorage, TSharedPtr<IHierarchyViewerDataInterface> InHierarchyData, TSharedPtr<IRowNode> InParent);
-		UE_API virtual ~FTopLevelRowsNode() override = default;
+		virtual ~FTopLevelRowsNode() override = default;
 		
 		/** Retrieve access to the rows used by this node. */
 		virtual FRowHandleArrayView GetRows() const override;
@@ -41,4 +41,4 @@ namespace UE::Editor::DataStorage::QueryStack
 	};
 }
 
-#undef UE_API
\ No newline at end of file
+#undef UE_API
diff --git a/Engine/Plugins/Experimental/EditorDataStorageFeatures/Source/TedsTableViewer/Public/TedsFilter.h b/Engine/Plugins/Experimental/EditorDataStorageFeatures/Source/TedsTableViewer/Public/TedsFilter.h
--- a/Engine/Plugins/Experimental/EditorDataStorageFeatures/Source/TedsTableViewer/Public/TedsFilter.h
+++ b/Engine/Plugins/Experimental/EditorDataStorageFeatures/Source/TedsTableViewer/Public/TedsFilter.h
@@ -13,6 +13,8 @@
 #include "Elements/Framework/TypedElementQueryContext.h"
 #include "Styling/SlateIconFinder.h"
 
+#define UE_API TEDSTABLEVIEWER_API
+
 #define LOCTEXT_NAMESPACE "TEDSFilter"
 
 namespace UE::Editor::DataStorage
@@ -33,17 +35,17 @@ namespace UE::Editor::DataStorage
 		return false;
 	}
 
-	struct TEDSTABLEVIEWER_API FTedsFilterData
+	struct FTedsFilterData
 	{
-		FTedsFilterData(const FName& InFilterName, const FText& InFilterDisplayName, const FText& InFilterToolTip, const FName& InFilterIconName,
+		UE_API FTedsFilterData(const FName& InFilterName, const FText& InFilterDisplayName, const FText& InFilterToolTip, const FName& InFilterIconName,
 			const TSharedPtr<FFilterCategory>& InCategory, const QueryHandle& InFilterQuery);
 
-		FTedsFilterData(const FName& InFilterName, const FText& InFilterDisplayName, const QueryHandle& InFilterQuery);
+		UE_API FTedsFilterData(const FName& InFilterName, const FText& InFilterDisplayName, const QueryHandle& InFilterQuery);
 		
-		FTedsFilterData(const FName& InFilterName, const FText& InFilterDisplayName, const FText& InFilterToolTip, const FName& InFilterIconName,
+		UE_API FTedsFilterData(const FName& InFilterName, const FText& InFilterDisplayName, const FText& InFilterToolTip, const FName& InFilterIconName,
 			const TSharedPtr<FFilterCategory>& InCategory, const Queries::TQueryFunction<bool>& InFilterQuery);
 
-		FTedsFilterData(const FName& InFilterName, const FText& InFilterDisplayName, const Queries::TQueryFunction<bool>& InFilterQuery);
+		UE_API FTedsFilterData(const FName& InFilterName, const FText& InFilterDisplayName, const Queries::TQueryFunction<bool>& InFilterQuery);
 
 		/*
 		 * Require that the template type cannot be one of the pre-defined FilterQuery types so it will pass to the defined
@@ -80,8 +82,8 @@ namespace UE::Editor::DataStorage
 				Queries::BuildQueryFunction<bool>(Forward<FilterFunction>(InFilterQuery)))
 		{}
 
-		FTedsFilterData(const FTedsFilterData&);
-		~FTedsFilterData(); 
+		UE_API FTedsFilterData(const FTedsFilterData&);
+		UE_API ~FTedsFilterData(); 
 		
 		FName FilterName;
 		FText FilterDisplayName;
@@ -204,4 +206,6 @@ namespace UE::Editor::DataStorage
 	};
 } // namespace UE::Editor::DataStorage
 
-#undef LOCTEXT_NAMESPACE
\ No newline at end of file
+#undef LOCTEXT_NAMESPACE
+
+#undef UE_API
diff --git a/Engine/Plugins/Experimental/EditorDataStorageFeatures/Source/TedsTableViewer/Public/Widgets/STedsFilterBar.h b/Engine/Plugins/Experimental/EditorDataStorageFeatures/Source/TedsTableViewer/Public/Widgets/STedsFilterBar.h
--- a/Engine/Plugins/Experimental/EditorDataStorageFeatures/Source/TedsTableViewer/Public/Widgets/STedsFilterBar.h
+++ b/Engine/Plugins/Experimental/EditorDataStorageFeatures/Source/TedsTableViewer/Public/Widgets/STedsFilterBar.h
@@ -6,6 +6,8 @@
 #include "TedsQueryStackInterfaces.h"
 #include "Elements/Framework/TypedElementQueryBuilder.h"
 
+#define UE_API TEDSTABLEVIEWER_API
+
 namespace UE::Editor::DataStorage
 {
 	namespace QueryStack
@@ -26,7 +28,7 @@ namespace UE::Editor::DataStorage
 	*
 	* Use MakeAddFilterButton() to make the button that summons the dropdown showing all the filters
 	*/
-	class TEDSTABLEVIEWER_API STedsFilterBar : public SBasicFilterBar<FTedsRowHandle&>
+	class STedsFilterBar : public SBasicFilterBar<FTedsRowHandle&>
 	{
 		using FOnFilterChanged = typename SBasicFilterBar<FTedsRowHandle&>::FOnFilterChanged;
 		DECLARE_DELEGATE(FOnPostFiltersChanged);
@@ -75,52 +77,52 @@ namespace UE::Editor::DataStorage
 		SLATE_END_ARGS()
 
 		/** Constructs this widget with InArgs */
-		void Construct(const FArguments& InArgs);
+		UE_API void Construct(const FArguments& InArgs);
 
 		/** Override tick so that filters are only recompiled once if multiple were added at once */
-		virtual void Tick(const FGeometry& AllottedGeometry, const double InCurrentTime, const float InDeltaTime) override;
+		UE_API virtual void Tick(const FGeometry& AllottedGeometry, const double InCurrentTime, const float InDeltaTime) override;
 
 		/** Add an array of new Filters to the menu post-initialization, will overwrite any filters with the same name */
-		TSharedRef<FTedsFilter> AddFilter(FTedsFilterData& InFilterData);
+		UE_API TSharedRef<FTedsFilter> AddFilter(FTedsFilterData& InFilterData);
 		/** Add a new Filter to the menu post-initialization, will overwrite any filters with the same name */
-		void AddFilters(const TArray<FTedsFilterData>& InFiltersData);
+		UE_API void AddFilters(const TArray<FTedsFilterData>& InFiltersData);
 		/** Remove a Filter from the menu and remove any place it may be referenced */
-		void DeleteFilter(const FName& InFilterName);
+		UE_API void DeleteFilter(const FName& InFilterName);
 		/** Remove all Filters from the menu and all references */
-		void DeleteAllFilters();
+		UE_API void DeleteAllFilters();
 		/** Set a new array of Common filters to appear in that section */
-		void SetCommonFilters(const TArray<FName>& InFilters);
+		UE_API void SetCommonFilters(const TArray<FName>& InFilters);
 		
 		/** 
 		 * Remove all filters and replace them with a new array of filters with an option to persist their state (is pinned & is active) for new
 		 * filters with the same FName identifier
 		 */
-		void ReplaceFilters(const TArray<FTedsFilterData>& InFilters, const bool bPersistFilterStates = true);
+		UE_API void ReplaceFilters(const TArray<FTedsFilterData>& InFilters, const bool bPersistFilterStates = true);
 		
 		// Add an external query description to the FilterBar
-		void AddExternalQuery(const FName& QueryName, const QueryHandle& InQueryHandle);
-		void RemoveExternalQuery(const FName& QueryName);
+		UE_API void AddExternalQuery(const FName& QueryName, const QueryHandle& InQueryHandle);
+		UE_API void RemoveExternalQuery(const FName& QueryName);
 
 		// Append all external queries into the given query description
-		void AppendExternalQueries(DataStorage::FQueryDescription& OutQuery);
+		UE_API void AppendExternalQueries(DataStorage::FQueryDescription& OutQuery);
 
 		// Add an external query function to the FilterBar
-		void AddExternalQueryFunction(const FName& QueryName, const Queries::TQueryFunction<bool>& InQueryFunction);
-		void RemoveExternalQueryFunction(const FName& QueryName);
+		UE_API void AddExternalQueryFunction(const FName& QueryName, const Queries::TQueryFunction<bool>& InQueryFunction);
+		UE_API void RemoveExternalQueryFunction(const FName& QueryName);
 
 		// Add an external class query function to the FilterBar (uses OR) 
-		void AddClassQueryFunction(const FName& ClassName, const Queries::TQueryFunction<bool>& InClassQueryFunction);
-		void RemoveClassQueryFunction(const FName& ClassName);
+		UE_API void AddClassQueryFunction(const FName& ClassName, const Queries::TQueryFunction<bool>& InClassQueryFunction);
+		UE_API void RemoveClassQueryFunction(const FName& ClassName);
 
 	protected:
-		void RecompileQueries();
+		UE_API void RecompileQueries();
 		
-		virtual void PopulateAddFilterMenu(UToolMenu* InMenu);
-		virtual void PopulateCommonFilterSections(UToolMenu* InMenu);
-		virtual void PopulateCustomFilters(UToolMenu* InMenu);
+		UE_API virtual void PopulateAddFilterMenu(UToolMenu* InMenu);
+		UE_API virtual void PopulateCommonFilterSections(UToolMenu* InMenu);
+		UE_API virtual void PopulateCustomFilters(UToolMenu* InMenu);
 	
 	private:
-		virtual TSharedRef<SWidget> MakeAddFilterMenu() override;
+		UE_API virtual TSharedRef<SWidget> MakeAddFilterMenu() override;
 
 		ICoreProvider* Storage = nullptr;
 
@@ -171,3 +173,5 @@ namespace UE::Editor::DataStorage
 		bool bMarkDirty = false;
 	};
 }
+
+#undef UE_API
diff --git a/Engine/Plugins/Experimental/EditorPerformance/Source/EditorPerformance/Public/EditorPerformanceModule.h b/Engine/Plugins/Experimental/EditorPerformance/Source/EditorPerformance/Public/EditorPerformanceModule.h
--- a/Engine/Plugins/Experimental/EditorPerformance/Source/EditorPerformance/Public/EditorPerformanceModule.h
+++ b/Engine/Plugins/Experimental/EditorPerformance/Source/EditorPerformance/Public/EditorPerformanceModule.h
@@ -60,7 +60,7 @@ public:
 	UE_API FEditorPerformanceModule::EEditorState GetEditorState() const;
 
 	DECLARE_EVENT(FEditorPerformanceModule, FPerformanceStateChangedEvent);
-	UE_API FPerformanceStateChangedEvent& GetOnPerformanceStateChanged() { return OnPerformanceStateChanged; }
+	FPerformanceStateChangedEvent& GetOnPerformanceStateChanged() { return OnPerformanceStateChanged; }
 
 private:
 	TSharedRef<SDockTab> CreatePerformanceReportTab(const FSpawnTabArgs& Args);	
diff --git a/Engine/Plugins/FX/Niagara/Source/NiagaraEditor/Public/ViewModels/Stack/NiagaraStackFunctionInput.h b/Engine/Plugins/FX/Niagara/Source/NiagaraEditor/Public/ViewModels/Stack/NiagaraStackFunctionInput.h
--- a/Engine/Plugins/FX/Niagara/Source/NiagaraEditor/Public/ViewModels/Stack/NiagaraStackFunctionInput.h
+++ b/Engine/Plugins/FX/Niagara/Source/NiagaraEditor/Public/ViewModels/Stack/NiagaraStackFunctionInput.h
@@ -359,7 +359,7 @@ public:
 	NIAGARAEDITOR_API UNiagaraStackFunctionInput* GetSharedLinkedValueInput();
 
 	/** Gets whether or not this input was created by another input to display and edit a shared linked value. */
-	NIAGARAEDITOR_API bool GetIsSharedLinkedValueInput() const { return bIsSharedLinkedValueInput; }
+	bool GetIsSharedLinkedValueInput() const { return bIsSharedLinkedValueInput; }
 
 	/** Moves the current value of this input, to a shared parameter and then changes the value of this input to link to that new parameter.  If a target usage is provided
 		the shared input parameter will be added to that script, otherwise the shared parameter will be added to a transient parameter in the current script. */
@@ -379,7 +379,7 @@ protected:
 	NIAGARAEDITOR_API FStackIssueFixDelegate GetUpgradeDynamicInputVersionFix();
 
 	/** Sets whether or not this input was created by another input to display and edit a shared linked value. */
-	NIAGARAEDITOR_API void SetIsSharedLinkedValueInput(bool bInIsSharedLinkedValueInput) { bIsSharedLinkedValueInput = bInIsSharedLinkedValueInput; }
+	void SetIsSharedLinkedValueInput(bool bInIsSharedLinkedValueInput) { bIsSharedLinkedValueInput = bInIsSharedLinkedValueInput; }
 
 private:
 	struct FInputValues
diff --git a/Engine/Plugins/Media/AppleProResMedia/Source/AppleProResMedia/Private/AppleProResEncoderProtocol.h b/Engine/Plugins/Media/AppleProResMedia/Source/AppleProResMedia/Private/AppleProResEncoderProtocol.h
--- a/Engine/Plugins/Media/AppleProResMedia/Source/AppleProResMedia/Private/AppleProResEncoderProtocol.h
+++ b/Engine/Plugins/Media/AppleProResMedia/Source/AppleProResMedia/Private/AppleProResEncoderProtocol.h
@@ -18,6 +18,8 @@
 
 #include "AppleProResEncoderProtocol.generated.h"
 
+#define UE_API APPLEPRORESMEDIA_API
+
 UENUM()
 enum class EAppleProResEncoderFormats : uint8
 {
@@ -45,20 +47,20 @@ enum class EAppleProResEncoderScanType : uint8
 	IM_INTERLATED_BOTTOM_FIRST_FIRST UMETA(DisplayName = "Interlaced mode; second (bottom) image line belongs to first temporal field"),
 };
 
-UCLASS(meta = (DisplayName = "Apple ProRes Encoder (mov)", CommandLineID = "AppleProRes") )
-class APPLEPRORESMEDIA_API UAppleProResEncoderProtocol : public UFrameGrabberProtocol
+UCLASS(MinimalAPI, meta = (DisplayName = "Apple ProRes Encoder (mov)", CommandLineID = "AppleProRes") )
+class UAppleProResEncoderProtocol : public UFrameGrabberProtocol
 {
 public:
 	GENERATED_BODY()
 
-	UAppleProResEncoderProtocol(const FObjectInitializer& InObjInit);
+	UE_API UAppleProResEncoderProtocol(const FObjectInitializer& InObjInit);
 
 public:
-	virtual bool SetupImpl() override;
-	virtual void FinalizeImpl() override;
-	virtual FFramePayloadPtr GetFramePayload(const FFrameMetrics& InFrameMetrics) override;
-	virtual void ProcessFrame(FCapturedFrameData InFrame) override;
-	virtual bool CanWriteToFileImpl(const TCHAR* InFilename, bool bInOverwriteExisting) const override;
+	UE_API virtual bool SetupImpl() override;
+	UE_API virtual void FinalizeImpl() override;
+	UE_API virtual FFramePayloadPtr GetFramePayload(const FFrameMetrics& InFrameMetrics) override;
+	UE_API virtual void ProcessFrame(FCapturedFrameData InFrame) override;
+	UE_API virtual bool CanWriteToFileImpl(const TCHAR* InFilename, bool bInOverwriteExisting) const override;
 
 public:
 
@@ -83,28 +85,28 @@ public:
 
 private:
 
-	void ParseCommandLine();
-	bool CreateProResFile(const FString& InFilename);
-	void ConvertFColorToRGBA4444(const TArray<FColor>& InColorbuffer);
-
-	PRCodecType GetSelectedCodecType() const;
-	PRVideoCodecType GetSelectedVideoCodecType() const;
-	ProResFormatDescriptionColorPrimaries GetColorPrimaries() const;
-	ProResFormatDescriptionTransferFunction GetTransferFunction() const;
-	ProResFormatDescriptionYCbCrMatrix GetYCbCrMatrix() const;
-	int32_t GetFrameHeaderColorPrimaries() const;
-	int32_t GetFrameHeaderTransferCharacteristic() const;
-	int32_t GetFrameHeaderMatrixCoefficients() const;
-	PRInterlaceMode GetInterlaceMode() const;
-	PRPixelFormat GetPixelFormat() const;
-	uint32 GetPixelAspectRatioHorizontalSpacing() const;
-	uint32 GetPixelAspectRatioVerticalSpacing() const;
-	uint32 GetStride() const;
-	uint32_t GetBitDepth() const;
-	uint32_t GetFieldCount() const;
-	ProResFormatDescriptionFieldDetail GetFieldDetail() const;
-	bool HasAlpha() const;
-	void ReleaseAndClearResources();
+	UE_API void ParseCommandLine();
+	UE_API bool CreateProResFile(const FString& InFilename);
+	UE_API void ConvertFColorToRGBA4444(const TArray<FColor>& InColorbuffer);
+
+	UE_API PRCodecType GetSelectedCodecType() const;
+	UE_API PRVideoCodecType GetSelectedVideoCodecType() const;
+	UE_API ProResFormatDescriptionColorPrimaries GetColorPrimaries() const;
+	UE_API ProResFormatDescriptionTransferFunction GetTransferFunction() const;
+	UE_API ProResFormatDescriptionYCbCrMatrix GetYCbCrMatrix() const;
+	UE_API int32_t GetFrameHeaderColorPrimaries() const;
+	UE_API int32_t GetFrameHeaderTransferCharacteristic() const;
+	UE_API int32_t GetFrameHeaderMatrixCoefficients() const;
+	UE_API PRInterlaceMode GetInterlaceMode() const;
+	UE_API PRPixelFormat GetPixelFormat() const;
+	UE_API uint32 GetPixelAspectRatioHorizontalSpacing() const;
+	UE_API uint32 GetPixelAspectRatioVerticalSpacing() const;
+	UE_API uint32 GetStride() const;
+	UE_API uint32_t GetBitDepth() const;
+	UE_API uint32_t GetFieldCount() const;
+	UE_API ProResFormatDescriptionFieldDetail GetFieldDetail() const;
+	UE_API bool HasAlpha() const;
+	UE_API void ReleaseAndClearResources();
 
 private:
 
@@ -123,3 +125,5 @@ private:
 	TArray<uint8> IntermediateFrameBuffer;
 	TArray<uint8> OutputFrameBuffer;
 };
+
+#undef UE_API
diff --git a/Engine/Plugins/Media/AppleProResMedia/Source/AppleProResMedia/Public/AppleProResEncoder/AppleProResEncoder.h b/Engine/Plugins/Media/AppleProResMedia/Source/AppleProResMedia/Public/AppleProResEncoder/AppleProResEncoder.h
--- a/Engine/Plugins/Media/AppleProResMedia/Source/AppleProResMedia/Public/AppleProResEncoder/AppleProResEncoder.h
+++ b/Engine/Plugins/Media/AppleProResMedia/Source/AppleProResMedia/Public/AppleProResEncoder/AppleProResEncoder.h
@@ -15,6 +15,8 @@
 
 #include "AppleProResEncoder.generated.h"
 
+#define UE_API APPLEPRORESMEDIA_API
+
 UENUM(BlueprintType)
 enum class EAppleProResEncoderCodec : uint8
 {
@@ -131,7 +133,7 @@ struct FAppleProResEncoderOptions
 	uint32 AudioChannelCount;
 };
 
-class APPLEPRORESMEDIA_API FAppleProResEncoder
+class FAppleProResEncoder
 {
 public:
 	struct FTimecodePayload : IImagePixelDataPayload
@@ -140,31 +142,31 @@ public:
 	};
 
 public:
-	FAppleProResEncoder(const FAppleProResEncoderOptions& InOptions);
-	~FAppleProResEncoder();
+	UE_API FAppleProResEncoder(const FAppleProResEncoderOptions& InOptions);
+	UE_API ~FAppleProResEncoder();
 
 	/** Call to initialize the Sink Writer. This must be done before attempting to write data to it. */
-	bool Initialize();
+	UE_API bool Initialize();
 
 	/** Finalize the video file and finish writing it to disk. Called by the destructor if not automatically called. */
-	void Finalize();
+	UE_API void Finalize();
 
 	/** Appends a new frame onto the output file. */
-	bool WriteFrame(const FImagePixelData* InPixelData);
+	UE_API bool WriteFrame(const FImagePixelData* InPixelData);
 	
 	/** Appends a new audio sample onto the audio stream. */
-	bool WriteAudioSample(const TArrayView<int16>& InAudioSamples);
+	UE_API bool WriteAudioSample(const TArrayView<int16>& InAudioSamples);
 
 	/** Gets the options that the encoder was initialized with. */
-	const FAppleProResEncoderOptions& GetOptions() const;
+	UE_API const FAppleProResEncoderOptions& GetOptions() const;
 
 	bool IsInitialized() const { return bInitialized; }
 
 private:
-	bool InitializeVideoTrack();
-	bool InitializeTimecodeTrack();
-	bool InitializeAudioTrack();
-	void ReleaseAndClearResources();
+	UE_API bool InitializeVideoTrack();
+	UE_API bool InitializeTimecodeTrack();
+	UE_API bool InitializeAudioTrack();
+	UE_API void ReleaseAndClearResources();
 
 private:
 	FAppleProResEncoderOptions Options;
@@ -195,4 +197,6 @@ private:
 	int32 MaxCompressedFrameSize;
 	int32 TargetCompressedFrameSize;
 	uint32 AudioSamplesPerFrame;
-};
\ No newline at end of file
+};
+
+#undef UE_API
diff --git a/Engine/Plugins/Media/PixelStreaming2/Source/PixelStreaming2/Internal/AudioCapturer.h b/Engine/Plugins/Media/PixelStreaming2/Source/PixelStreaming2/Internal/AudioCapturer.h
--- a/Engine/Plugins/Media/PixelStreaming2/Source/PixelStreaming2/Internal/AudioCapturer.h
+++ b/Engine/Plugins/Media/PixelStreaming2/Source/PixelStreaming2/Internal/AudioCapturer.h
@@ -20,7 +20,7 @@ namespace UE::PixelStreaming2
 	{
 	public:
 		UE_API FAudioPatchMixer(uint8 NumChannels, uint32 SampleRate, float SampleSizeSeconds);
-		UE_API virtual ~FAudioPatchMixer() = default;
+		virtual ~FAudioPatchMixer() = default;
 
 		UE_API uint32 GetMaxBufferSize() const;
 		UE_API uint8  GetNumChannels() const;
diff --git a/Engine/Plugins/Media/PixelStreaming2/Source/PixelStreaming2RTC/Internal/EpicRtcAudioCapturer.h b/Engine/Plugins/Media/PixelStreaming2/Source/PixelStreaming2RTC/Internal/EpicRtcAudioCapturer.h
--- a/Engine/Plugins/Media/PixelStreaming2/Source/PixelStreaming2RTC/Internal/EpicRtcAudioCapturer.h
+++ b/Engine/Plugins/Media/PixelStreaming2/Source/PixelStreaming2RTC/Internal/EpicRtcAudioCapturer.h
@@ -15,13 +15,13 @@ namespace UE::PixelStreaming2
 	class FEpicRtcAudioCapturer : public FAudioCapturer
 	{
 	public:
-		UE_API virtual ~FEpicRtcAudioCapturer() = default;
+		virtual ~FEpicRtcAudioCapturer() = default;
 
 		// Override the push audio method as EpicRtc needs the broadcasted audio to be in 10ms chunks
 		UE_API virtual void PushAudio(const float* AudioData, int32 NumSamples, int32 NumChannels, int32 SampleRate) override;
 
 	protected:
-		UE_API FEpicRtcAudioCapturer(const int SampleRate = 48000, const int NumChannels = 2, const float SampleSizeInSeconds = 0.5f)
+		FEpicRtcAudioCapturer(const int SampleRate = 48000, const int NumChannels = 2, const float SampleSizeInSeconds = 0.5f)
 			: FAudioCapturer(SampleRate, NumChannels, SampleSizeInSeconds) {}
 
 	private:
diff --git a/Engine/Plugins/MovieScene/MovieRenderPipeline/Source/MovieRenderPipelineCore/Public/Graph/MovieGraphRenderLayerSubsystem.h b/Engine/Plugins/MovieScene/MovieRenderPipeline/Source/MovieRenderPipelineCore/Public/Graph/MovieGraphRenderLayerSubsystem.h
--- a/Engine/Plugins/MovieScene/MovieRenderPipeline/Source/MovieRenderPipelineCore/Public/Graph/MovieGraphRenderLayerSubsystem.h
+++ b/Engine/Plugins/MovieScene/MovieRenderPipeline/Source/MovieRenderPipelineCore/Public/Graph/MovieGraphRenderLayerSubsystem.h
@@ -883,15 +883,15 @@ class UMovieGraphModifierBase : public UObject
 public:
 	/** Applies this modifier in the provided world. Called once per layer. */
 	UFUNCTION(BlueprintCallable, Category = "Modifier")
-	MOVIERENDERPIPELINECORE_API virtual void ApplyModifier(const UWorld* World) PURE_VIRTUAL(UMovieGraphModifierBase::ApplyModifier, );
+	virtual void ApplyModifier(const UWorld* World) PURE_VIRTUAL(UMovieGraphModifierBase::ApplyModifier, );
 
 	/** Undoes the effects of this modifier. The modifier is responsible for tracking changes that it makes. Called once per layer. */
 	UFUNCTION(BlueprintCallable, Category = "Modifier")
-	MOVIERENDERPIPELINECORE_API virtual void UndoModifier() PURE_VIRTUAL(UMovieGraphModifierBase::UndoModifier, );
+	virtual void UndoModifier() PURE_VIRTUAL(UMovieGraphModifierBase::UndoModifier, );
 
 	/** Gets the name of this modifier. Typically used for UI display purposes. */
 	UFUNCTION(BlueprintCallable, Category = "Modifier")
-	MOVIERENDERPIPELINECORE_API virtual FText GetModifierName() PURE_VIRTUAL(UMovieGraphModifierBase::GetModifierName, return FText::GetEmpty(););
+	virtual FText GetModifierName() PURE_VIRTUAL(UMovieGraphModifierBase::GetModifierName, return FText::GetEmpty(););
 };
 
 /**
diff --git a/Engine/Plugins/MovieScene/MovieRenderPipeline/Source/MovieRenderPipelineCore/Public/Graph/Nodes/MovieGraphBurnInNode.h b/Engine/Plugins/MovieScene/MovieRenderPipeline/Source/MovieRenderPipelineCore/Public/Graph/Nodes/MovieGraphBurnInNode.h
--- a/Engine/Plugins/MovieScene/MovieRenderPipeline/Source/MovieRenderPipelineCore/Public/Graph/Nodes/MovieGraphBurnInNode.h
+++ b/Engine/Plugins/MovieScene/MovieRenderPipeline/Source/MovieRenderPipelineCore/Public/Graph/Nodes/MovieGraphBurnInNode.h
@@ -107,7 +107,7 @@ class UMovieGraphOutputBurnInNode : public UMovieGraphBurnInNode
 	GENERATED_BODY()
 
 public:
-	UE_API UMovieGraphOutputBurnInNode() = default;
+	UMovieGraphOutputBurnInNode() = default;
 
 	// UMovieGraphSettingNode Interface
 	virtual FString GetNodeInstanceName() const override;
diff --git a/Engine/Plugins/MovieScene/MovieRenderPipeline/Source/MovieRenderPipelineCore/Public/Graph/Nodes/MovieGraphModifierNode.h b/Engine/Plugins/MovieScene/MovieRenderPipeline/Source/MovieRenderPipelineCore/Public/Graph/Nodes/MovieGraphModifierNode.h
--- a/Engine/Plugins/MovieScene/MovieRenderPipeline/Source/MovieRenderPipelineCore/Public/Graph/Nodes/MovieGraphModifierNode.h
+++ b/Engine/Plugins/MovieScene/MovieRenderPipeline/Source/MovieRenderPipelineCore/Public/Graph/Nodes/MovieGraphModifierNode.h
@@ -52,21 +52,21 @@ class IMovieGraphModifierNodeInterface : public IInterface
 public:
 	/** Gets all modifiers that will be applied with this modifier node. */
 	UFUNCTION(BlueprintCallable, Category = "Modifiers")
-	UE_API virtual TArray<UMovieGraphModifierBase*> GetAllModifiers() const
+	virtual TArray<UMovieGraphModifierBase*> GetAllModifiers() const
 	{
 		return {};
 	}
 
 	/** Gets whether this modifier node supports/uses collections. */
 	UFUNCTION(BlueprintCallable, Category = "Modifiers")
-	UE_API virtual bool SupportsCollections() const
+	virtual bool SupportsCollections() const
 	{
 		return true;
 	}
 
 	/** Gets the enable state (within this modifier) of the collection with the given name. */
 	UFUNCTION(BlueprintCallable, Category = "Modifiers")
-	UE_API virtual bool IsCollectionEnabled(const FName& InCollectionName) const
+	virtual bool IsCollectionEnabled(const FName& InCollectionName) const
 	{
 		return true;
 	}
@@ -76,28 +76,28 @@ public:
 	 * node. Collections that are added to the modifier are enabled by default.
 	 */
 	UFUNCTION(BlueprintCallable, Category = "Modifiers")
-	UE_API virtual void SetCollectionEnabled(const FName& InCollectionName, const bool bIsCollectionEnabled)
+	virtual void SetCollectionEnabled(const FName& InCollectionName, const bool bIsCollectionEnabled)
 	{
 		
 	}
 
 	/** Gets all collections that will be affected by the modifiers used by this node. */
 	UFUNCTION(BlueprintCallable, Category = "Modifiers")
-	UE_API virtual TArray<FName> GetAllCollections() const
+	virtual TArray<FName> GetAllCollections() const
 	{
 		return {};
 	}
 
 	/** Adds a collection identified by the given name which will be affected by the modifiers on this node. */
 	UFUNCTION(BlueprintCallable, Category = "Modifiers")
-	UE_API virtual void AddCollection(const FName& InCollectionName)
+	virtual void AddCollection(const FName& InCollectionName)
 	{
 		
 	}
 
 	/** Removes a collection identified by the given name. Returns true if the collection was found and removed successfully, else false. */
 	UFUNCTION(BlueprintCallable, Category = "Modifiers")
-	UE_API virtual bool RemoveCollection(const FName& InCollectionName)
+	virtual bool RemoveCollection(const FName& InCollectionName)
 	{
 		return false;
 	}
diff --git a/Engine/Plugins/MovieScene/MovieRenderPipeline/Source/MovieRenderPipelineCore/Public/Graph/Nodes/MovieGraphRenderPassNode.h b/Engine/Plugins/MovieScene/MovieRenderPipeline/Source/MovieRenderPipelineCore/Public/Graph/Nodes/MovieGraphRenderPassNode.h
--- a/Engine/Plugins/MovieScene/MovieRenderPipeline/Source/MovieRenderPipelineCore/Public/Graph/Nodes/MovieGraphRenderPassNode.h
+++ b/Engine/Plugins/MovieScene/MovieRenderPipeline/Source/MovieRenderPipelineCore/Public/Graph/Nodes/MovieGraphRenderPassNode.h
@@ -74,7 +74,7 @@ public:
 	 */
 	TArray<FSoftClassPath> GetOutputTypeRestrictions() const { return GetOutputTypeRestrictionsImpl(); }
 
-	UE_API virtual void ResolveTokenContainingProperties(TFunction<void(FString&)>& ResolveFunc, const FMovieGraphTokenResolveContext& InContext) override
+	virtual void ResolveTokenContainingProperties(TFunction<void(FString&)>& ResolveFunc, const FMovieGraphTokenResolveContext& InContext) override
 	{
 		Super::ResolveTokenContainingProperties(ResolveFunc, InContext);
 
@@ -127,4 +127,4 @@ protected:
 	virtual TArray<FSoftClassPath> GetOutputTypeRestrictionsImpl() const { return {}; } 
 };
 
-#undef UE_API
\ No newline at end of file
+#undef UE_API
diff --git a/Engine/Plugins/Runtime/MeshModelingToolset/Source/MeshModelingToolsEditorOnly/Public/SkeletalMesh/SkinWeightsPaintTool.h b/Engine/Plugins/Runtime/MeshModelingToolset/Source/MeshModelingToolsEditorOnly/Public/SkeletalMesh/SkinWeightsPaintTool.h
--- a/Engine/Plugins/Runtime/MeshModelingToolset/Source/MeshModelingToolsEditorOnly/Public/SkeletalMesh/SkinWeightsPaintTool.h
+++ b/Engine/Plugins/Runtime/MeshModelingToolset/Source/MeshModelingToolsEditorOnly/Public/SkeletalMesh/SkinWeightsPaintTool.h
@@ -780,7 +780,7 @@ public:
 	UE_API USkinWeightsPaintToolProperties* GetWeightToolProperties() const;
 
 	// get access to the mesh being edited
-	UE_API FDynamicMesh3& GetMesh() {return EditedMesh;};
+	FDynamicMesh3& GetMesh() {return EditedMesh;};
 
 	// get the set of bones whose weights are actually locked, taking into account bBypassAllLocks
 	UE_API const TSet<int32>& GetFinalLockedBoneIndices();
diff --git a/Engine/Plugins/Runtime/Metasound/Source/MetasoundEditor/Public/MetasoundEditorGraphSchema.h b/Engine/Plugins/Runtime/Metasound/Source/MetasoundEditor/Public/MetasoundEditorGraphSchema.h
--- a/Engine/Plugins/Runtime/Metasound/Source/MetasoundEditor/Public/MetasoundEditorGraphSchema.h
+++ b/Engine/Plugins/Runtime/Metasound/Source/MetasoundEditor/Public/MetasoundEditorGraphSchema.h
@@ -187,7 +187,7 @@ struct FMetasoundGraphSchemaAction_NewInput : public FMetasoundGraphSchemaAction
 
 	UE_API virtual const FLinearColor& GetIconColor() const override;
 
-	UE_API virtual FString GetUniqueStableFavoriteKey() const override
+	virtual FString GetUniqueStableFavoriteKey() const override
 	{
 		return NodeID.ToString();
 	}
@@ -237,7 +237,7 @@ struct FMetasoundGraphSchemaAction_NewOutput : public FMetasoundGraphSchemaActio
 
 	UE_API virtual const FLinearColor& GetIconColor() const override;
 
-	UE_API virtual FString GetUniqueStableFavoriteKey() const override
+	virtual FString GetUniqueStableFavoriteKey() const override
 	{
 		return NodeID.ToString();
 	}
@@ -292,7 +292,7 @@ struct FMetasoundGraphSchemaAction_NewVariableNode : public FMetasoundGraphSchem
 
 	UE_API virtual const FLinearColor& GetIconColor() const override;
 
-	UE_API virtual FString GetUniqueStableFavoriteKey() const override
+	virtual FString GetUniqueStableFavoriteKey() const override
 	{
 		return VariableID.ToString();
 	}
diff --git a/Engine/Plugins/Runtime/StateTree/Source/StateTreeEditorModule/Public/StateTreeEditorPropertyBindings.h b/Engine/Plugins/Runtime/StateTree/Source/StateTreeEditorModule/Public/StateTreeEditorPropertyBindings.h
--- a/Engine/Plugins/Runtime/StateTree/Source/StateTreeEditorModule/Public/StateTreeEditorPropertyBindings.h
+++ b/Engine/Plugins/Runtime/StateTree/Source/StateTreeEditorModule/Public/StateTreeEditorPropertyBindings.h
@@ -120,7 +120,7 @@ PRAGMA_DISABLE_DEPRECATION_WARNINGS
 	}
 
 	UE_DEPRECATED(5.6, "Use FindBinding taking FPropertyBindingPath instead")
-	UE_API const FStateTreePropertyPathBinding* FindPropertyBinding(const FStateTreePropertyPath& TargetPath, ESearchMode SearchMode = ESearchMode::Exact) const
+	const FStateTreePropertyPathBinding* FindPropertyBinding(const FStateTreePropertyPath& TargetPath, ESearchMode SearchMode = ESearchMode::Exact) const
 	{
 		return static_cast<const FStateTreePropertyPathBinding*>(FindBinding(TargetPath, SearchMode));
 	}
diff --git a/Engine/Plugins/Runtime/StateTree/Source/StateTreeEditorModule/Public/StateTreeEditorSchema.h b/Engine/Plugins/Runtime/StateTree/Source/StateTreeEditorModule/Public/StateTreeEditorSchema.h
--- a/Engine/Plugins/Runtime/StateTree/Source/StateTreeEditorModule/Public/StateTreeEditorSchema.h
+++ b/Engine/Plugins/Runtime/StateTree/Source/StateTreeEditorModule/Public/StateTreeEditorSchema.h
@@ -48,7 +48,7 @@ public:
 	 * Validates and applies the schema restrictions on the StateTree.
 	 * This is before the engine validation. The EditorData will be correctly set.
 	 */
-	UE_API virtual void PreValidate(TNotNull<UStateTree*> StateTree)
+	virtual void PreValidate(TNotNull<UStateTree*> StateTree)
 	{
 	}
 
diff --git a/Engine/Plugins/TextureGraph/Source/TextureGraphEngine/Job/JobArgs.h b/Engine/Plugins/TextureGraph/Source/TextureGraphEngine/Job/JobArgs.h
--- a/Engine/Plugins/TextureGraph/Source/TextureGraphEngine/Job/JobArgs.h
+++ b/Engine/Plugins/TextureGraph/Source/TextureGraphEngine/Job/JobArgs.h
@@ -144,7 +144,7 @@ public:
 	UE_API virtual CHashPtr				TileHash(uint32 TileX, uint32 TileY) const override;
 
 	UE_API virtual bool					IsLateBound(uint32 TileX, uint32 TileY) const override;
-	UE_API virtual bool					IsBlob() const override { return true; }
+	virtual bool					IsBlob() const override { return true; }
 
 	UE_API virtual JobPtrW					GeneratingJob() const override;
 	UE_API virtual const BufferDescriptor* GetDescriptor() const override;
diff --git a/Engine/Plugins/VirtualProduction/CameraCalibrationCore/Source/CameraCalibrationCore/Public/AnamorphicLensDistortionModelHandler.h b/Engine/Plugins/VirtualProduction/CameraCalibrationCore/Source/CameraCalibrationCore/Public/AnamorphicLensDistortionModelHandler.h
--- a/Engine/Plugins/VirtualProduction/CameraCalibrationCore/Source/CameraCalibrationCore/Public/AnamorphicLensDistortionModelHandler.h
+++ b/Engine/Plugins/VirtualProduction/CameraCalibrationCore/Source/CameraCalibrationCore/Public/AnamorphicLensDistortionModelHandler.h
@@ -26,7 +26,7 @@ protected:
 	UE_API virtual void InterpretDistortionParameters() override;
 	UE_API virtual FString GetDistortionShaderPath() const override;
 	UE_API virtual void ExecuteDistortionShader(class FRDGBuilder& GraphBuilder, const FLensDistortionState& InCurrentState, float InverseOverscan, float CameraOverscan, const FVector2D& SensorSize, FRDGTextureRef& OutDistortionMap) const override;
-	UE_API virtual bool IsForwardDistorting() const override { return false; }
+	virtual bool IsForwardDistorting() const override { return false; }
 	//~ End ULensDistortionModelHandlerBase interface
 
 private:
diff --git a/Engine/Plugins/VirtualProduction/CameraCalibrationCore/Source/CameraCalibrationCore/Public/BrownConradyDULensDistortionModelHandler.h b/Engine/Plugins/VirtualProduction/CameraCalibrationCore/Source/CameraCalibrationCore/Public/BrownConradyDULensDistortionModelHandler.h
--- a/Engine/Plugins/VirtualProduction/CameraCalibrationCore/Source/CameraCalibrationCore/Public/BrownConradyDULensDistortionModelHandler.h
+++ b/Engine/Plugins/VirtualProduction/CameraCalibrationCore/Source/CameraCalibrationCore/Public/BrownConradyDULensDistortionModelHandler.h
@@ -21,11 +21,11 @@ class UBrownConradyDULensDistortionModelHandler : public UBrownConradyUDLensDist
 protected:
 	//~ Begin ULensDistortionModelHandlerBase interface
 	UE_API virtual void InitializeHandler() override;
-	UE_API virtual bool IsForwardDistorting() const override { return false; }
+	virtual bool IsForwardDistorting() const override { return false; }
 	UE_API virtual FVector2D ComputeDistortedUV(const FVector2D& InScreenUV) const override;
 	UE_API virtual FVector2D ComputeUndistortedUV(const FVector2D& InScreenUV) const override;
 	UE_API virtual void InitDistortionMaterials() override;
 	//~ End ULensDistortionModelHandlerBase interface
 };
 
-#undef UE_API
\ No newline at end of file
+#undef UE_API
diff --git a/Engine/Plugins/VirtualProduction/CameraCalibrationCore/Source/CameraCalibrationCore/Public/CameraCalibrationCharucoBoard.h b/Engine/Plugins/VirtualProduction/CameraCalibrationCore/Source/CameraCalibrationCore/Public/CameraCalibrationCharucoBoard.h
--- a/Engine/Plugins/VirtualProduction/CameraCalibrationCore/Source/CameraCalibrationCore/Public/CameraCalibrationCharucoBoard.h
+++ b/Engine/Plugins/VirtualProduction/CameraCalibrationCore/Source/CameraCalibrationCore/Public/CameraCalibrationCharucoBoard.h
@@ -85,7 +85,7 @@ public:
 
 	/** Get the generated Charuco board texture */
 	UFUNCTION(BlueprintCallable, Category = "Calibration")
-	UE_API UTexture2D* GetCharucoBoardTexture() const { return CharucoBoardTexture; }
+	UTexture2D* GetCharucoBoardTexture() const { return CharucoBoardTexture; }
 
 	/** Get the total number of corners in the Charuco board */
 	UFUNCTION(BlueprintCallable, Category = "Calibration")
@@ -190,4 +190,4 @@ private:
 	TObjectPtr<UMaterialInstanceDynamic> DynamicBoardMaterial;
 };
 
-#undef UE_API
\ No newline at end of file
+#undef UE_API
diff --git a/Engine/Plugins/VirtualProduction/Takes/Source/TakeRecorderSources/Public/TakeRecorderWorldSource.h b/Engine/Plugins/VirtualProduction/Takes/Source/TakeRecorderSources/Public/TakeRecorderWorldSource.h
--- a/Engine/Plugins/VirtualProduction/Takes/Source/TakeRecorderSources/Public/TakeRecorderWorldSource.h
+++ b/Engine/Plugins/VirtualProduction/Takes/Source/TakeRecorderSources/Public/TakeRecorderWorldSource.h
@@ -49,7 +49,7 @@ private:
 	// UTakeRecorderSource
 	UE_API virtual TArray<UTakeRecorderSource*> PreRecording(ULevelSequence* InSequence, FMovieSceneSequenceID InSequenceID, ULevelSequence* InRootSequence, FManifestSerializer* InManifestSerializer) override;
 	UE_API virtual TArray<UTakeRecorderSource*> PostRecording(class ULevelSequence* InSequence, class ULevelSequence* InRootSequence, const bool bCancelled) override;
-	UE_API virtual bool SupportsTakeNumber() const override { return false; }
+	virtual bool SupportsTakeNumber() const override { return false; }
 	UE_API virtual FText GetDisplayTextImpl() const override;
 	UE_API virtual FText GetAddSourceDisplayTextImpl() const override;		
 	UE_API virtual bool CanAddSource(UTakeRecorderSources* InSources) const override;
@@ -67,4 +67,4 @@ private:
 	TWeakObjectPtr<UTakeRecorderActorSource> WorldSource;
 };
 
-#undef UE_API
\ No newline at end of file
+#undef UE_API
diff --git a/Engine/Source/Editor/CurveEditor/Public/Selection/ScopedSelectionChangeEventSuppression.h b/Engine/Source/Editor/CurveEditor/Public/Selection/ScopedSelectionChangeEventSuppression.h
--- a/Engine/Source/Editor/CurveEditor/Public/Selection/ScopedSelectionChangeEventSuppression.h
+++ b/Engine/Source/Editor/CurveEditor/Public/Selection/ScopedSelectionChangeEventSuppression.h
@@ -5,6 +5,8 @@
 #include "Templates/SharedPointer.h"
 #include "Templates/UnrealTemplate.h"
 
+#define UE_API CURVEEDITOR_API
+
 class FCurveEditor;
 
 namespace UE::CurveEditor
@@ -31,12 +33,12 @@ namespace UE::CurveEditor
  *		// ~FScopedSelectionChangeEventSuppression will now invoke FCurveEditorSelection::OnSelectionChanged once exactly.
  * }
  */
-class CURVEEDITOR_API FScopedSelectionChangeEventSuppression : public FNoncopyable
+class FScopedSelectionChangeEventSuppression : public FNoncopyable
 {
 public:
 
-	explicit FScopedSelectionChangeEventSuppression(const TSharedRef<FCurveEditor>& InCurveEditor);
-	~FScopedSelectionChangeEventSuppression();
+	UE_API explicit FScopedSelectionChangeEventSuppression(const TSharedRef<FCurveEditor>& InCurveEditor);
+	UE_API ~FScopedSelectionChangeEventSuppression();
 
 private:
 
@@ -45,3 +47,4 @@ private:
 };
 }
 
+#undef UE_API
diff --git a/Engine/Source/Editor/DataHierarchyEditor/Public/DataHierarchyViewModelBase.h b/Engine/Source/Editor/DataHierarchyEditor/Public/DataHierarchyViewModelBase.h
--- a/Engine/Source/Editor/DataHierarchyEditor/Public/DataHierarchyViewModelBase.h
+++ b/Engine/Source/Editor/DataHierarchyEditor/Public/DataHierarchyViewModelBase.h
@@ -868,10 +868,10 @@ struct FHierarchyElementViewModel : TSharedFromThis<FHierarchyElementViewModel>,
 	virtual bool DoesExternalDataStillExist(const UHierarchyDataRefreshContext* Context) const { return false; }
 
 	/** The UObject we display in the details panel when this item is selected. By default it's the hierarchy element the view model represents. */
-	UE_API virtual UObject* GetObjectForEditing() { return Element; }
+	virtual UObject* GetObjectForEditing() { return Element; }
 	/** Source items are transient, which is why we don't allow editing by default.
 	 * This is useful to override if source data points at actual data to edit. */
-	UE_API virtual bool AllowEditingInDetailsPanel() const { return bIsForHierarchy; }
+	virtual bool AllowEditingInDetailsPanel() const { return bIsForHierarchy; }
 
 	bool IsForHierarchy() const { return bIsForHierarchy; }
 
@@ -895,7 +895,7 @@ protected:
 	
 	/** Called after an element was dropped onto the calling element. Allows inserting specialized fix up. Not called if OnDroppedOnOverride was called.
 	 *  Note: DroppedElementViewModel is the newly created view model _after_ the drop occured. */
-	UE_API virtual void PostOnDroppedOn(TSharedPtr<FHierarchyElementViewModel> DroppedElementViewModel) { }
+	virtual void PostOnDroppedOn(TSharedPtr<FHierarchyElementViewModel> DroppedElementViewModel) { }
 
 private:
 	/** Should return true if this element can contain other elements as children. Can optionally specify a message for failure. */
@@ -1035,7 +1035,7 @@ struct FHierarchySectionViewModel : FHierarchyElementViewModel
 	UE_API FText GetSectionTooltip() const;
 
 	void SetSectionImage(const FSlateBrush* InSectionImage) { SectionImage = InSectionImage; }
-	UE_API virtual const FSlateBrush* GetSectionImageBrush() const { return SectionImage ? SectionImage : FAppStyle::GetNoBrush(); }
+	virtual const FSlateBrush* GetSectionImageBrush() const { return SectionImage ? SectionImage : FAppStyle::GetNoBrush(); }
 	
 	bool IsAllSection() const { return GetData()->GetFName() == UDataHierarchyViewModelBase::AllSectionHierarchyObjectName || GetData()->GetFName() == UDataHierarchyViewModelBase::AllSectionSourceObjectName; }
 protected:
diff --git a/Engine/Source/Editor/Experimental/EditorInteractiveToolsFramework/Public/ViewportInteractions/ViewportCameraTranslateInteraction.h b/Engine/Source/Editor/Experimental/EditorInteractiveToolsFramework/Public/ViewportInteractions/ViewportCameraTranslateInteraction.h
--- a/Engine/Source/Editor/Experimental/EditorInteractiveToolsFramework/Public/ViewportInteractions/ViewportCameraTranslateInteraction.h
+++ b/Engine/Source/Editor/Experimental/EditorInteractiveToolsFramework/Public/ViewportInteractions/ViewportCameraTranslateInteraction.h
@@ -27,8 +27,8 @@ public:
 	UE_API UViewportCameraTranslateInteraction();
 
 	//~ Begin IKeyInputBehaviorTarget
-	UE_API virtual void OnKeyPressed(const FKey& InKeyID) override {}
-	UE_API virtual void OnKeyReleased(const FKey& InKeyID) override {}
+	virtual void OnKeyPressed(const FKey& InKeyID) override {}
+	virtual void OnKeyReleased(const FKey& InKeyID) override {}
 	UE_API virtual void OnForceEndCapture() override;
 	//~ End IKeyInputBehaviorTarget
 
diff --git a/Engine/Source/Editor/Experimental/EditorInteractiveToolsFramework/Public/ViewportInteractions/ViewportClickInteraction.h b/Engine/Source/Editor/Experimental/EditorInteractiveToolsFramework/Public/ViewportInteractions/ViewportClickInteraction.h
--- a/Engine/Source/Editor/Experimental/EditorInteractiveToolsFramework/Public/ViewportInteractions/ViewportClickInteraction.h
+++ b/Engine/Source/Editor/Experimental/EditorInteractiveToolsFramework/Public/ViewportInteractions/ViewportClickInteraction.h
@@ -6,6 +6,8 @@
 #include "BaseBehaviors/BehaviorTargetInterfaces.h"
 #include "ViewportClickInteraction.generated.h"
 
+#define UE_API EDITORINTERACTIVETOOLSFRAMEWORK_API
+
 namespace UE::InteractiveTools
 {
 enum class EViewportClickButton : uint8
@@ -23,19 +25,19 @@ struct FViewportClick;
 /**
  *
  */
-UCLASS(Abstract)
-class EDITORINTERACTIVETOOLSFRAMEWORK_API UViewportClickInteraction
+UCLASS(MinimalAPI, Abstract)
+class UViewportClickInteraction
 	: public UViewportInteraction
 	, public IClickBehaviorTarget
 {
 	GENERATED_BODY()
 
 public:
-	UViewportClickInteraction();
+	UE_API UViewportClickInteraction();
 
 	// ~Begin IClickBehaviorTarget
-	virtual FInputRayHit IsHitByClick(const FInputDeviceRay& InClickPos) override;
-	virtual void OnClicked(const FInputDeviceRay& InClickPos) override;
+	UE_API virtual FInputRayHit IsHitByClick(const FInputDeviceRay& InClickPos) override;
+	UE_API virtual void OnClicked(const FInputDeviceRay& InClickPos) override;
 	// ~End IClickBehaviorTarget
 
 	FKey MouseButton = EKeys::LeftMouseButton;
@@ -44,10 +46,12 @@ protected:
 	/** Override this function to provide custom functionality */
 	virtual void ProcessClick_Internal(const FViewportClick& InViewportClick, HHitProxy* InHitProxy, FSceneView& View) {}
 
-	virtual FDeviceButtonState GetActiveMouseButtonState(const FInputDeviceState& Input);
+	UE_API virtual FDeviceButtonState GetActiveMouseButtonState(const FInputDeviceState& Input);
 
-	FKey GetMouseKey() const;
-	virtual EInputEvent GetMouseEvent() const;
+	UE_API FKey GetMouseKey() const;
+	UE_API virtual EInputEvent GetMouseEvent() const;
 
 	TWeakObjectPtr<USingleClickInputBehavior> SingleClickInputBehaviorWeak;
 };
+
+#undef UE_API
diff --git a/Engine/Source/Editor/Experimental/EditorInteractiveToolsFramework/Public/ViewportInteractions/ViewportDragInteraction.h b/Engine/Source/Editor/Experimental/EditorInteractiveToolsFramework/Public/ViewportInteractions/ViewportDragInteraction.h
--- a/Engine/Source/Editor/Experimental/EditorInteractiveToolsFramework/Public/ViewportInteractions/ViewportDragInteraction.h
+++ b/Engine/Source/Editor/Experimental/EditorInteractiveToolsFramework/Public/ViewportInteractions/ViewportDragInteraction.h
@@ -36,8 +36,8 @@ public:
 	UE_API virtual void OnDragStart(const FInputDeviceRay& InDragStartPos) override;
 	UE_API virtual void OnDrag(const FInputDeviceRay& InDragPos) override;
 	UE_API virtual void OnDragEnd(const FInputDeviceRay& InDragEndPos) override;
-	UE_API virtual void OnTerminateDragSequence() override {}
-	UE_API virtual void OnStateUpdated(const FInputDeviceState& InInputDeviceState) override {}
+	virtual void OnTerminateDragSequence() override {}
+	virtual void OnStateUpdated(const FInputDeviceState& InInputDeviceState) override {}
 	UE_API virtual void OnForceEndCapture() override;
 	//~ End IViewportClickDragBehaviorTarget
 
diff --git a/Engine/Source/Editor/LevelEditor/Private/ViewportInteractions/LevelViewportClickSelection.h b/Engine/Source/Editor/LevelEditor/Private/ViewportInteractions/LevelViewportClickSelection.h
--- a/Engine/Source/Editor/LevelEditor/Private/ViewportInteractions/LevelViewportClickSelection.h
+++ b/Engine/Source/Editor/LevelEditor/Private/ViewportInteractions/LevelViewportClickSelection.h
@@ -6,46 +6,50 @@
 #include "ViewportInteractions/ViewportClickInteraction.h"
 #include "LevelViewportClickSelection.generated.h"
 
+#define UE_API LEVELEDITOR_API
+
 /**
  * Left-mouse click in level editor viewport, for elements selection
  */
-UCLASS()
-class LEVELEDITOR_API ULevelViewportClickSelection : public UViewportClickInteraction
+UCLASS(MinimalAPI)
+class ULevelViewportClickSelection : public UViewportClickInteraction
 {
 	GENERATED_BODY()
 
 public:
-	ULevelViewportClickSelection();
+	UE_API ULevelViewportClickSelection();
 
 protected:
-	virtual void ProcessClick_Internal(const FViewportClick& InViewportClick, HHitProxy* InHitProxy, FSceneView& View) override;
+	UE_API virtual void ProcessClick_Internal(const FViewportClick& InViewportClick, HHitProxy* InHitProxy, FSceneView& View) override;
 };
 
 /**
  * Right-mouse click in level editor viewport, for selecting + summoning context menu
  */
-UCLASS()
-class LEVELEDITOR_API ULevelViewportClickContextMenu : public ULevelViewportClickSelection
+UCLASS(MinimalAPI)
+class ULevelViewportClickContextMenu : public ULevelViewportClickSelection
 {
 	GENERATED_BODY()
 
 public:
-	ULevelViewportClickContextMenu();
+	UE_API ULevelViewportClickContextMenu();
 };
 
 /**
  * Left-mouse double click
  */
-UCLASS()
-class LEVELEDITOR_API ULevelViewportDoubleClick : public ULevelViewportClickSelection
+UCLASS(MinimalAPI)
+class ULevelViewportDoubleClick : public ULevelViewportClickSelection
 {
 	GENERATED_BODY()
 
 public:
-	ULevelViewportDoubleClick();
+	UE_API ULevelViewportDoubleClick();
 
 	virtual EInputEvent GetMouseEvent() const override
 	{
 		return IE_DoubleClick;
 	}
 };
+
+#undef UE_API
diff --git a/Engine/Source/Editor/MaterialEditor/Internal/DiffTool/Widgets/SMaterialDiff.h b/Engine/Source/Editor/MaterialEditor/Internal/DiffTool/Widgets/SMaterialDiff.h
--- a/Engine/Source/Editor/MaterialEditor/Internal/DiffTool/Widgets/SMaterialDiff.h
+++ b/Engine/Source/Editor/MaterialEditor/Internal/DiffTool/Widgets/SMaterialDiff.h
@@ -9,6 +9,8 @@
 #include "SBlueprintDiff.h"
 #include "Widgets/SCompoundWidget.h"
 
+#define UE_API MATERIALEDITOR_API
+
 class FBlueprintDifferenceTreeEntry;
 class UEdGraph;
 class UMaterialGraph;
@@ -24,7 +26,7 @@ struct FMaterialDiffResultItem : public FDiffResultItem
 };
 
 /* Visual Diff between two Material Graphs */
-class MATERIALEDITOR_API SMaterialDiff : public SCompoundWidget
+class SMaterialDiff : public SCompoundWidget
 {
 public:
 	SLATE_BEGIN_ARGS(SMaterialDiff) {}
@@ -36,82 +38,82 @@ public:
 		SLATE_ARGUMENT(TSharedPtr<SWindow>, ParentWindow)
 	SLATE_END_ARGS()
 
-	void Construct(const FArguments& InArgs);
-	virtual ~SMaterialDiff() override;
+	UE_API void Construct(const FArguments& InArgs);
+	UE_API virtual ~SMaterialDiff() override;
 
 	/** Called when a new Graph is clicked on by user */
-	void OnGraphChanged(FMaterialToDiff* Diff);
+	UE_API void OnGraphChanged(FMaterialToDiff* Diff);
 
 	/** Called when user clicks on a new graph list item */
-	void OnGraphSelectionChanged(TSharedPtr<FMaterialToDiff> Item, ESelectInfo::Type SelectionType);
+	UE_API void OnGraphSelectionChanged(TSharedPtr<FMaterialToDiff> Item, ESelectInfo::Type SelectionType);
 
 	/** Called when user clicks on an entry in the listview of differences */
-	void OnDiffListSelectionChanged(TSharedPtr<FMaterialDiffResultItem> TheDiff);
+	UE_API void OnDiffListSelectionChanged(TSharedPtr<FMaterialDiffResultItem> TheDiff);
 
 	/** Helper function for generating an empty widget */
-	static TSharedRef<SWidget> DefaultEmptyPanel();
+	static UE_API TSharedRef<SWidget> DefaultEmptyPanel();
 
 	/** Helper function to create a window that holds a diff widget */
-	static TSharedPtr<SWindow> CreateDiffWindow(FText WindowTitle, TObjectPtr<UMaterialGraph> OldMaterialGraph, TObjectPtr<UMaterialGraph> NewMaterialGraph, const FRevisionInfo& OldRevision, const FRevisionInfo& NewRevision);
+	static UE_API TSharedPtr<SWindow> CreateDiffWindow(FText WindowTitle, TObjectPtr<UMaterialGraph> OldMaterialGraph, TObjectPtr<UMaterialGraph> NewMaterialGraph, const FRevisionInfo& OldRevision, const FRevisionInfo& NewRevision);
 
 	/** Helper function to create a window that holds a diff widget (default window title) */
-	static TSharedPtr<SWindow> CreateDiffWindow(TObjectPtr<UMaterialGraph> OldMaterialGraph, TObjectPtr<UMaterialGraph> NewMaterialGraph, const FRevisionInfo& OldRevision, const FRevisionInfo& NewRevision, const UClass* ObjectClass);
+	static UE_API TSharedPtr<SWindow> CreateDiffWindow(TObjectPtr<UMaterialGraph> OldMaterialGraph, TObjectPtr<UMaterialGraph> NewMaterialGraph, const FRevisionInfo& OldRevision, const FRevisionInfo& NewRevision, const UClass* ObjectClass);
 
 	DECLARE_MULTICAST_DELEGATE_OneParam(FOnWindowClosedEvent, TSharedRef<SMaterialDiff>)
 	FOnWindowClosedEvent OnWindowClosedEvent;
 
 protected:
 	/** Called when user clicks button to go to next difference */
-	void NextDiff();
+	UE_API void NextDiff();
 
 	/** Called when user clicks button to go to prev difference */
-	void PrevDiff();
+	UE_API void PrevDiff();
 
 	/** Called to determine whether we have a list of differences to cycle through */
-	bool HasNextDiff() const;
-	bool HasPrevDiff() const;
+	UE_API bool HasNextDiff() const;
+	UE_API bool HasPrevDiff() const;
 
 	/** User toggles the option to show/hide preview viewports */
-	void ToggleViewport();
+	UE_API void ToggleViewport();
 
 	/** Get the image to show for the toggle viewport option */
-	FSlateIcon GetViewportImage() const;
+	UE_API FSlateIcon GetViewportImage() const;
 
 	/** Find the FMaterialToDiff that displays the graph with GraphPath relative path */
-	FMaterialToDiff* FindGraphToDiffEntry(const FString& GraphPath);
+	UE_API FMaterialToDiff* FindGraphToDiffEntry(const FString& GraphPath);
 
 	/** Bring these revisions of graph into focus on main display */
-	void FocusOnMaterialGraphRevisions(FMaterialToDiff* Diff);
+	UE_API void FocusOnMaterialGraphRevisions(FMaterialToDiff* Diff);
 
 	/** User toggles the option to lock the views between the two material graphs */
-	void OnToggleLockView();
+	UE_API void OnToggleLockView();
 
 	/** User toggles the option to change the split view mode betwwen vertical and horizontal */
-	void OnToggleSplitViewMode();
+	UE_API void OnToggleSplitViewMode();
 
 	/** Get the image to show for the toggle lock option */
-	FSlateIcon GetLockViewImage() const;
+	UE_API FSlateIcon GetLockViewImage() const;
 
 	/** Get the image to show for the toggle split view mode option */
-	FSlateIcon GetSplitViewModeImage() const;
+	UE_API FSlateIcon GetSplitViewModeImage() const;
 
 	/** Reset the graph editor, called when user switches graphs to display */
-	void ResetGraphEditors();
+	UE_API void ResetGraphEditors();
 
 	/** Material Graph to diff, is added to panel */
 	TSharedPtr<FMaterialToDiff> MaterialGraphToDiff;
 
 	/** Get Graph editor associated with this Graph */
-	FMaterialDiffPanel& GetDiffPanelForNode(UEdGraphNode& Node);
+	UE_API FMaterialDiffPanel& GetDiffPanelForNode(UEdGraphNode& Node);
 
 	/** Event handler that updates the graph view when user selects a new graph */
-	void HandleGraphChanged(const FString& GraphPath);
+	UE_API void HandleGraphChanged(const FString& GraphPath);
 
 	/** Function used to generate the list of differences and the widgets needed to calculate that list */
-	void GenerateDifferencesList();
+	UE_API void GenerateDifferencesList();
 
 	/** Called when editor may need to be closed */
-	void OnCloseAssetEditor(UObject* Asset, EAssetEditorCloseReason CloseReason);
+	UE_API void OnCloseAssetEditor(UObject* Asset, EAssetEditorCloseReason CloseReason);
 
 	struct FMaterialGraphDiffControl
 	{
@@ -123,22 +125,22 @@ protected:
 		TSharedPtr<IDiffControl> DiffControl;
 	};
 
-	FMaterialGraphDiffControl GenerateMaterialGraphPanel();
+	UE_API FMaterialGraphDiffControl GenerateMaterialGraphPanel();
 
-	TSharedRef<SOverlay> GenerateMaterialGraphWidgetForPanel(FMaterialDiffPanel& OutDiffPanel) const;
-	TSharedRef<SBox> GenerateRevisionInfoWidgetForPanel(TSharedPtr<SWidget>& OutGeneratedWidget, const FText& InRevisionText) const;
+	UE_API TSharedRef<SOverlay> GenerateMaterialGraphWidgetForPanel(FMaterialDiffPanel& OutDiffPanel) const;
+	UE_API TSharedRef<SBox> GenerateRevisionInfoWidgetForPanel(TSharedPtr<SWidget>& OutGeneratedWidget, const FText& InRevisionText) const;
 
 	/** Accessor and event handler for toggling between diff view modes - only GraphMode for now: */
-	void SetCurrentMode(FName NewMode);
+	UE_API void SetCurrentMode(FName NewMode);
 
 	FName GetCurrentMode() const { return CurrentMode; }
 
-	void OnModeChanged(const FName& InNewViewMode) const;
+	UE_API void OnModeChanged(const FName& InNewViewMode) const;
 
-	void UpdateTopSectionVisibility(const FName& InNewViewMode) const;
+	UE_API void UpdateTopSectionVisibility(const FName& InNewViewMode) const;
 
-	void SetCurrentWidgetIndex(int32 Index);
-	int32 GetCurrentWidgetIndex() const;
+	UE_API void SetCurrentWidgetIndex(int32 Index);
+	UE_API int32 GetCurrentWidgetIndex() const;
 
 	FName CurrentMode;
 
@@ -183,3 +185,5 @@ protected:
 
 	FDelegateHandle AssetEditorCloseDelegate;
 };
+
+#undef UE_API
diff --git a/Engine/Source/Editor/MaterialEditor/Internal/DiffTool/Widgets/SMaterialInstanceDiff.h b/Engine/Source/Editor/MaterialEditor/Internal/DiffTool/Widgets/SMaterialInstanceDiff.h
--- a/Engine/Source/Editor/MaterialEditor/Internal/DiffTool/Widgets/SMaterialInstanceDiff.h
+++ b/Engine/Source/Editor/MaterialEditor/Internal/DiffTool/Widgets/SMaterialInstanceDiff.h
@@ -11,13 +11,15 @@
 #include "DiffTool/MaterialInstanceDiffPanel.h"
 #include "Widgets/SCompoundWidget.h"
 
+#define UE_API MATERIALEDITOR_API
+
 class SDetailsSplitter;
 class FBlueprintDifferenceTreeEntry;
 
 enum class EAssetEditorCloseReason : uint8;
 
 /* Visual Diff between two Material Instances */
-class MATERIALEDITOR_API SMaterialInstanceDiff : public SCompoundWidget
+class SMaterialInstanceDiff : public SCompoundWidget
 {
 public:
 	SLATE_BEGIN_ARGS(SMaterialInstanceDiff) {}
@@ -28,46 +30,46 @@ public:
 		SLATE_ARGUMENT(TSharedPtr<SWindow>, ParentWindow)
 	SLATE_END_ARGS()
 
-	void Construct(const FArguments& InArgs);
-	virtual ~SMaterialInstanceDiff() override;
+	UE_API void Construct(const FArguments& InArgs);
+	UE_API virtual ~SMaterialInstanceDiff() override;
 
 	/** Called when user clicks on an entry in the listview of differences */
-	void OnDiffListSelectionChanged(TSharedPtr<FMaterialDiffResultItem> TheDiff);
+	UE_API void OnDiffListSelectionChanged(TSharedPtr<FMaterialDiffResultItem> TheDiff);
 
 	/** Helper function for generating an empty widget */
-	static TSharedRef<SWidget> DefaultEmptyPanel();
+	static UE_API TSharedRef<SWidget> DefaultEmptyPanel();
 
 	/** Helper function to create a window that holds a diff widget */
-	static TSharedPtr<SWindow> CreateDiffWindow(FText WindowTitle, TObjectPtr<UMaterialInstance>  OldMaterialInstance, TObjectPtr<UMaterialInstance>  NewMaterialInstance, const FRevisionInfo& OldRevision, const FRevisionInfo& NewRevision);
+	static UE_API TSharedPtr<SWindow> CreateDiffWindow(FText WindowTitle, TObjectPtr<UMaterialInstance>  OldMaterialInstance, TObjectPtr<UMaterialInstance>  NewMaterialInstance, const FRevisionInfo& OldRevision, const FRevisionInfo& NewRevision);
 
 	/** Helper function to create a window that holds a diff widget (default window title) */
-	static TSharedPtr<SWindow> CreateDiffWindow(TObjectPtr<UMaterialInstance>  OldMaterialInstance, TObjectPtr<UMaterialInstance>  NewMaterialInstance, const FRevisionInfo& OldRevision, const FRevisionInfo& NewRevision, const UClass* ObjectClass);
+	static UE_API TSharedPtr<SWindow> CreateDiffWindow(TObjectPtr<UMaterialInstance>  OldMaterialInstance, TObjectPtr<UMaterialInstance>  NewMaterialInstance, const FRevisionInfo& OldRevision, const FRevisionInfo& NewRevision, const UClass* ObjectClass);
 
 	DECLARE_MULTICAST_DELEGATE_OneParam(FOnWindowClosedEvent, TSharedRef<SMaterialInstanceDiff>)
 	FOnWindowClosedEvent OnWindowClosedEvent;
 
 protected:
 	/** Called when user clicks button to go to next difference */
-	void NextDiff();
+	UE_API void NextDiff();
 
 	/** Called when user clicks button to go to prev difference */
-	void PrevDiff();
+	UE_API void PrevDiff();
 
 	/** Called to determine whether we have a list of differences to cycle through */
-	bool HasNextDiff() const;
-	bool HasPrevDiff() const;
+	UE_API bool HasNextDiff() const;
+	UE_API bool HasPrevDiff() const;
 
 	/** User toggles the option to show/hide preview viewports */
-	void ToggleViewport();
+	UE_API void ToggleViewport();
 
 	/** Get the image to show for the toggle viewport option */
-	FSlateIcon GetViewportImage() const;
+	UE_API FSlateIcon GetViewportImage() const;
 
 	/** Function used to generate the list of differences and the widgets needed to calculate that list */
-	void GenerateDifferencesList();
+	UE_API void GenerateDifferencesList();
 
 	/** Called when editor may need to be closed */
-	void OnCloseAssetEditor(UObject* Asset, EAssetEditorCloseReason CloseReason);
+	UE_API void OnCloseAssetEditor(UObject* Asset, EAssetEditorCloseReason CloseReason);
 
 	struct FMaterialInstanceDiffControl
 	{
@@ -81,14 +83,14 @@ protected:
 		TSharedPtr<IDiffControl> DiffControl;
 	};
 
-	FMaterialInstanceDiffControl GenerateMaterialInstancePanel();
+	UE_API FMaterialInstanceDiffControl GenerateMaterialInstancePanel();
 
 	/** Accessor and event handler for toggling between diff view modes - only details mode for now: */
-	void SetCurrentMode(FName NewMode);
+	UE_API void SetCurrentMode(FName NewMode);
 
 	FName GetCurrentMode() const { return CurrentMode; }
 
-	void OnModeChanged(const FName& InNewViewMode) const;
+	UE_API void OnModeChanged(const FName& InNewViewMode) const;
 
 	FName CurrentMode;
 
@@ -121,3 +123,5 @@ protected:
 
 	FDelegateHandle AssetEditorCloseDelegate;
 };
+
+#undef UE_API
diff --git a/Engine/Source/Editor/SceneOutliner/Public/SSceneOutliner.h b/Engine/Source/Editor/SceneOutliner/Public/SSceneOutliner.h
--- a/Engine/Source/Editor/SceneOutliner/Public/SSceneOutliner.h
+++ b/Engine/Source/Editor/SceneOutliner/Public/SSceneOutliner.h
@@ -433,7 +433,7 @@ public:
 	SceneOutliner::FOnItemSelectionChanged& GetOnItemSelectionChanged() { return OnItemSelectionChanged; }
 
 	/** Event that's broadcast after the scene outliner has been re-populated. */
-	UE_API FSimpleMulticastDelegate& GetOnFullRefreshDone() { return OnFullRefreshDone; }
+	FSimpleMulticastDelegate& GetOnFullRefreshDone() { return OnFullRefreshDone; }
 
 	/** Set the item selection of the outliner based on a selector function. Any items which return true will be added */
 	UE_API virtual void SetSelection(const TFunctionRef<bool(ISceneOutlinerTreeItem&)> Selector) override;
diff --git a/Engine/Source/Editor/Sequencer/Public/MVVM/ViewModels/SequencerEditorViewModel.h b/Engine/Source/Editor/Sequencer/Public/MVVM/ViewModels/SequencerEditorViewModel.h
--- a/Engine/Source/Editor/Sequencer/Public/MVVM/ViewModels/SequencerEditorViewModel.h
+++ b/Engine/Source/Editor/Sequencer/Public/MVVM/ViewModels/SequencerEditorViewModel.h
@@ -68,7 +68,7 @@ public:
 	/** Gets the current hotspots across any of our track areas */
 	UE_API TSharedPtr<ITrackAreaHotspot> GetHotspot() const;
 
-	UE_API void HandleDataHierarchyChanged() {}
+	void HandleDataHierarchyChanged() {}
 
 protected:
 
diff --git a/Engine/Source/Editor/StructUtilsEditor/Public/PropertyBagDetails.h b/Engine/Source/Editor/StructUtilsEditor/Public/PropertyBagDetails.h
--- a/Engine/Source/Editor/StructUtilsEditor/Public/PropertyBagDetails.h
+++ b/Engine/Source/Editor/StructUtilsEditor/Public/PropertyBagDetails.h
@@ -288,7 +288,7 @@ class UPropertyBagDetailsExtension : public UObject
 	GENERATED_BODY()
 
 public:
-	UE_API virtual TSharedRef<FPropertyBagInstanceDataDetails> CreateInstanceDataDetails(const FPropertyBagInstanceDataDetails::FConstructParams& Params) const
+	virtual TSharedRef<FPropertyBagInstanceDataDetails> CreateInstanceDataDetails(const FPropertyBagInstanceDataDetails::FConstructParams& Params) const
 	{
 		return MakeShared<FPropertyBagInstanceDataDetails>(Params);
 	}
diff --git a/Engine/Source/Editor/UnrealEd/Public/EditorViewportClient.h b/Engine/Source/Editor/UnrealEd/Public/EditorViewportClient.h
--- a/Engine/Source/Editor/UnrealEd/Public/EditorViewportClient.h
+++ b/Engine/Source/Editor/UnrealEd/Public/EditorViewportClient.h
@@ -673,7 +673,7 @@ public:
 	 * Override to specify what should happen when a context menu is requested
 	 * @param HitProxyElement the clicked element
 	 */
-	UNREALED_API virtual void SummonContextMenu(const FTypedElementHandle& HitProxyElement = FTypedElementHandle()) {}
+	virtual void SummonContextMenu(const FTypedElementHandle& HitProxyElement = FTypedElementHandle()) {}
 
 	/**
 	 * Called when mouse movement tracking begins
diff --git a/Engine/Source/Editor/UnrealEd/Public/Tools/EdModeInteractiveToolsContext.h b/Engine/Source/Editor/UnrealEd/Public/Tools/EdModeInteractiveToolsContext.h
--- a/Engine/Source/Editor/UnrealEd/Public/Tools/EdModeInteractiveToolsContext.h
+++ b/Engine/Source/Editor/UnrealEd/Public/Tools/EdModeInteractiveToolsContext.h
@@ -375,7 +375,7 @@ public:
 	 * Mark Editor Modes Tools for Drag Tools Support
 	 */
 	UE_DEPRECATED(5.8, "Drag tools have been merged with VIewport Interactions")
-	UE_API void SetDragToolsEnabled(bool bInEnabled) { SetViewportInteractionsEnabled(bInEnabled); }
+	void SetDragToolsEnabled(bool bInEnabled) { SetViewportInteractionsEnabled(bInEnabled); }
 
 	/**
 	 * Enable ITF Input tools support (e.g. Viewport Interactions as camera movement and selection)
@@ -390,7 +390,7 @@ public:
 	 * Enable/Disable support by using SetDragToolsEnabled(...)
 	 */
 	UE_DEPRECATED(5.8, "Drag tools have been merged with VIewport Interactions")
-	UE_API bool UsesDragTools() const { return UsesViewportInteractions(); } 
+	bool UsesDragTools() const { return UsesViewportInteractions(); } 
 
 	/**
 	 * Are Viewport Interactions supported? e.g. camera movements, commands, clicking, etc
@@ -431,7 +431,7 @@ protected:
 	 * Activating Drag Tools if current Mode Manager requires them
 	 */
 	UE_DEPRECATED(5.8, "Drag tools have been merged with viewport interactions.")
-	UE_API void RegisterDragTools() {};
+	void RegisterDragTools() {};
 
 	/**
 	 * Activating Viewport Interactions if current Mode Manager requires them
@@ -442,7 +442,7 @@ protected:
 	 * Deactivating Drag Tools
 	 */
 	UE_DEPRECATED(5.8, "Drag tools have been merged with viewport interactions.")
-	UE_API void UnregisterDragTools() {};
+	void UnregisterDragTools() {};
 
 	/**
 	 * Deactivating Viewport Interactions
diff --git a/Engine/Source/Editor/UserAssetTagsEditor/Private/AssetEditors/AssetEditor_TaggedAssetBrowserConfiguration.h b/Engine/Source/Editor/UserAssetTagsEditor/Private/AssetEditors/AssetEditor_TaggedAssetBrowserConfiguration.h
--- a/Engine/Source/Editor/UserAssetTagsEditor/Private/AssetEditors/AssetEditor_TaggedAssetBrowserConfiguration.h
+++ b/Engine/Source/Editor/UserAssetTagsEditor/Private/AssetEditors/AssetEditor_TaggedAssetBrowserConfiguration.h
@@ -6,22 +6,26 @@
 #include "Tools/UAssetEditor.h"
 #include "AssetEditor_TaggedAssetBrowserConfiguration.generated.h"
 
+#define UE_API USERASSETTAGSEDITOR_API
+
 /**
  * 
  */
-UCLASS(Transient)
-class USERASSETTAGSEDITOR_API UAssetEditor_TaggedAssetBrowserConfiguration : public UAssetEditor
+UCLASS(MinimalAPI, Transient)
+class UAssetEditor_TaggedAssetBrowserConfiguration : public UAssetEditor
 {
 	GENERATED_BODY()
 
 public:
-	void SetObjectToEdit(UObject* InObject);
+	UE_API void SetObjectToEdit(UObject* InObject);
 
 protected:
-	virtual void GetObjectsToEdit(TArray<UObject*>& OutObjectsToEdit) override;
-	virtual TSharedPtr<FBaseAssetToolkit> CreateToolkit() override;
+	UE_API virtual void GetObjectsToEdit(TArray<UObject*>& OutObjectsToEdit) override;
+	UE_API virtual TSharedPtr<FBaseAssetToolkit> CreateToolkit() override;
 
 private:
 	UPROPERTY(Transient)
 	TObjectPtr<UObject> ObjectToEdit;
 };
+
+#undef UE_API
diff --git a/Engine/Source/Editor/UserAssetTagsEditor/Private/Assets/AssetDefinition_TaggedAssetBrowserConfiguration.h b/Engine/Source/Editor/UserAssetTagsEditor/Private/Assets/AssetDefinition_TaggedAssetBrowserConfiguration.h
--- a/Engine/Source/Editor/UserAssetTagsEditor/Private/Assets/AssetDefinition_TaggedAssetBrowserConfiguration.h
+++ b/Engine/Source/Editor/UserAssetTagsEditor/Private/Assets/AssetDefinition_TaggedAssetBrowserConfiguration.h
@@ -6,17 +6,21 @@
 #include "AssetDefinitionDefault.h"
 #include "AssetDefinition_TaggedAssetBrowserConfiguration.generated.h"
 
+#define UE_API USERASSETTAGSEDITOR_API
+
 /**
  * 
  */
-UCLASS()
-class USERASSETTAGSEDITOR_API UAssetDefinition_TaggedAssetBrowserConfiguration : public UAssetDefinitionDefault
+UCLASS(MinimalAPI)
+class UAssetDefinition_TaggedAssetBrowserConfiguration : public UAssetDefinitionDefault
 {
 	GENERATED_BODY()
 
-	virtual FText GetAssetDisplayName() const override;
-	virtual TSoftClassPtr<UObject> GetAssetClass() const override;
-	virtual FLinearColor GetAssetColor() const override;
-	virtual EAssetCommandResult OpenAssets(const FAssetOpenArgs& OpenArgs) const override;
-	virtual TConstArrayView<FAssetCategoryPath> GetAssetCategories() const override;
+	UE_API virtual FText GetAssetDisplayName() const override;
+	UE_API virtual TSoftClassPtr<UObject> GetAssetClass() const override;
+	UE_API virtual FLinearColor GetAssetColor() const override;
+	UE_API virtual EAssetCommandResult OpenAssets(const FAssetOpenArgs& OpenArgs) const override;
+	UE_API virtual TConstArrayView<FAssetCategoryPath> GetAssetCategories() const override;
 };
+
+#undef UE_API
diff --git a/Engine/Source/Editor/UserAssetTagsEditor/Private/UserAssetTagProviders/UserAssetTagProvider_LocalAssetTypeFavorites.h b/Engine/Source/Editor/UserAssetTagsEditor/Private/UserAssetTagProviders/UserAssetTagProvider_LocalAssetTypeFavorites.h
--- a/Engine/Source/Editor/UserAssetTagsEditor/Private/UserAssetTagProviders/UserAssetTagProvider_LocalAssetTypeFavorites.h
+++ b/Engine/Source/Editor/UserAssetTagsEditor/Private/UserAssetTagProviders/UserAssetTagProvider_LocalAssetTypeFavorites.h
@@ -6,26 +6,30 @@
 #include "UserAssetTagProvider.h"
 #include "UserAssetTagProvider_LocalAssetTypeFavorites.generated.h"
 
+#define UE_API USERASSETTAGSEDITOR_API
+
 class UUserAssetTagEditorContext;
 
 /**
  * The local favorites provider lets you add your own selection of favorite tags per asset type
  * for quick reuse without affecting project settings.
  */
-UCLASS(DisplayName="Favorites by type")
-class USERASSETTAGSEDITOR_API UUserAssetTagProvider_LocalAssetTypeFavorites : public UUserAssetTagProvider
+UCLASS(MinimalAPI, DisplayName="Favorites by type")
+class UUserAssetTagProvider_LocalAssetTypeFavorites : public UUserAssetTagProvider
 {
 	GENERATED_BODY()
 
 public:
-	UUserAssetTagProvider_LocalAssetTypeFavorites();
+	UE_API UUserAssetTagProvider_LocalAssetTypeFavorites();
 
-	virtual FText GetDisplayNameText(const UUserAssetTagEditorContext* Context) const override;
-	virtual TSet<FName> GetSuggestedUserAssetTags(const UUserAssetTagEditorContext* Context) const override;
+	UE_API virtual FText GetDisplayNameText(const UUserAssetTagEditorContext* Context) const override;
+	UE_API virtual TSet<FName> GetSuggestedUserAssetTags(const UUserAssetTagEditorContext* Context) const override;
 
 	/** Favorites are per asset type, so we disable this provider if we encounter multiple asset types. */
-	virtual FResultWithUserFeedback IsValid(const UUserAssetTagEditorContext* Context) const override;
+	UE_API virtual FResultWithUserFeedback IsValid(const UUserAssetTagEditorContext* Context) const override;
 
-	virtual TSharedPtr<SWidget> AddAdditionalSuggestedWidgets(FName UserAssetTag, const UUserAssetTagEditorContext* Context) const override;
-	virtual TSharedPtr<SWidget> AddAdditionalOwnedTagWidgets(FName UserAssetTag, const UUserAssetTagEditorContext* Context) const override;
+	UE_API virtual TSharedPtr<SWidget> AddAdditionalSuggestedWidgets(FName UserAssetTag, const UUserAssetTagEditorContext* Context) const override;
+	UE_API virtual TSharedPtr<SWidget> AddAdditionalOwnedTagWidgets(FName UserAssetTag, const UUserAssetTagEditorContext* Context) const override;
 };
+
+#undef UE_API
diff --git a/Engine/Source/Editor/UserAssetTagsEditor/Private/UserAssetTagProviders/UserAssetTagProvider_Project.h b/Engine/Source/Editor/UserAssetTagsEditor/Private/UserAssetTagProviders/UserAssetTagProvider_Project.h
--- a/Engine/Source/Editor/UserAssetTagsEditor/Private/UserAssetTagProviders/UserAssetTagProvider_Project.h
+++ b/Engine/Source/Editor/UserAssetTagsEditor/Private/UserAssetTagProviders/UserAssetTagProvider_Project.h
@@ -6,20 +6,24 @@
 #include "UserAssetTagProvider.h"
 #include "UserAssetTagProvider_Project.generated.h"
 
+#define UE_API USERASSETTAGSEDITOR_API
+
 /** 
  *  A provider with the commonly used tags for this project, per asset type.
  *  You can assign tags per asset type in the project settings.
  */
-UCLASS(DisplayName="Project Tags by type")
-class USERASSETTAGSEDITOR_API UUserAssetTagProvider_Project : public UUserAssetTagProvider
+UCLASS(MinimalAPI, DisplayName="Project Tags by type")
+class UUserAssetTagProvider_Project : public UUserAssetTagProvider
 {
 	GENERATED_BODY()
 
-	virtual FText GetDisplayNameText(const UUserAssetTagEditorContext* Context) const override;
-	virtual TSet<FName> GetSuggestedUserAssetTags(const UUserAssetTagEditorContext* Context) const override;
-	virtual FResultWithUserFeedback IsValid(const UUserAssetTagEditorContext* Context) const override;
+	UE_API virtual FText GetDisplayNameText(const UUserAssetTagEditorContext* Context) const override;
+	UE_API virtual TSet<FName> GetSuggestedUserAssetTags(const UUserAssetTagEditorContext* Context) const override;
+	UE_API virtual FResultWithUserFeedback IsValid(const UUserAssetTagEditorContext* Context) const override;
 	
-	static void NavigateToProjectTags();
+	static UE_API void NavigateToProjectTags();
 protected:
-	virtual void AddToolbarMenuEntries(class UToolMenu* DynamicMenu, const UUserAssetTagEditorContext* Context) const override;
+	UE_API virtual void AddToolbarMenuEntries(class UToolMenu* DynamicMenu, const UUserAssetTagEditorContext* Context) const override;
 };
+
+#undef UE_API
diff --git a/Engine/Source/Editor/UserAssetTagsEditor/Public/AssetFactories/TaggedAssetBrowserConfigurationFactory.h b/Engine/Source/Editor/UserAssetTagsEditor/Public/AssetFactories/TaggedAssetBrowserConfigurationFactory.h
--- a/Engine/Source/Editor/UserAssetTagsEditor/Public/AssetFactories/TaggedAssetBrowserConfigurationFactory.h
+++ b/Engine/Source/Editor/UserAssetTagsEditor/Public/AssetFactories/TaggedAssetBrowserConfigurationFactory.h
@@ -6,12 +6,16 @@
 #include "Factories/Factory.h"
 #include "TaggedAssetBrowserConfigurationFactory.generated.h"
 
-UCLASS()
-class USERASSETTAGSEDITOR_API UTaggedAssetBrowserConfigurationFactory : public UFactory
+#define UE_API USERASSETTAGSEDITOR_API
+
+UCLASS(MinimalAPI)
+class UTaggedAssetBrowserConfigurationFactory : public UFactory
 {
 	GENERATED_BODY()
 
-	UTaggedAssetBrowserConfigurationFactory();
+	UE_API UTaggedAssetBrowserConfigurationFactory();
 	
-	virtual UObject* FactoryCreateNew(UClass* InClass, UObject* InParent, FName InName, EObjectFlags Flags, UObject* Context, FFeedbackContext* Warn) override;
+	UE_API virtual UObject* FactoryCreateNew(UClass* InClass, UObject* InParent, FName InName, EObjectFlags Flags, UObject* Context, FFeedbackContext* Warn) override;
 };
+
+#undef UE_API
diff --git a/Engine/Source/Editor/UserAssetTagsEditor/Public/Assets/TaggedAssetBrowserConfiguration.h b/Engine/Source/Editor/UserAssetTagsEditor/Public/Assets/TaggedAssetBrowserConfiguration.h
--- a/Engine/Source/Editor/UserAssetTagsEditor/Public/Assets/TaggedAssetBrowserConfiguration.h
+++ b/Engine/Source/Editor/UserAssetTagsEditor/Public/Assets/TaggedAssetBrowserConfiguration.h
@@ -6,6 +6,8 @@
 #include "UObject/Object.h"
 #include "TaggedAssetBrowserConfiguration.generated.h"
 
+#define UE_API USERASSETTAGSEDITOR_API
+
 class UTaggedAssetBrowserFilterBase;
 class UTaggedAssetBrowserFilterRoot;
 
@@ -41,13 +43,13 @@ struct FTaggedAssetBrowserConfigurationData_Extension : public FTaggedAssetBrows
 };
 
 /** The actual configuration asset. Contains a hierarchy root that lets you add new elements and has drag & drop support via Data Hierarchy Editor. */
-UCLASS()
-class USERASSETTAGSEDITOR_API UTaggedAssetBrowserConfiguration : public UObject
+UCLASS(MinimalAPI)
+class UTaggedAssetBrowserConfiguration : public UObject
 {
 	GENERATED_BODY()
 public:
 
-	UTaggedAssetBrowserConfiguration();
+	UE_API UTaggedAssetBrowserConfiguration();
 
 	/** Used for saving filter state and various menus. Should be unique. Is also used for extension identification. */
 	UPROPERTY(EditAnywhere, Category="Configuration", AssetRegistrySearchable)
@@ -68,6 +70,7 @@ public:
 	TObjectPtr<UTaggedAssetBrowserFilterRoot> FilterRoot;
 
 protected:
-	virtual EDataValidationResult IsDataValid(class FDataValidationContext& Context) const override;
+	UE_API virtual EDataValidationResult IsDataValid(class FDataValidationContext& Context) const override;
 };
 
+#undef UE_API
diff --git a/Engine/Source/Editor/UserAssetTagsEditor/Public/Config/LocalFavoriteUserAssetTagsConfig.h b/Engine/Source/Editor/UserAssetTagsEditor/Public/Config/LocalFavoriteUserAssetTagsConfig.h
--- a/Engine/Source/Editor/UserAssetTagsEditor/Public/Config/LocalFavoriteUserAssetTagsConfig.h
+++ b/Engine/Source/Editor/UserAssetTagsEditor/Public/Config/LocalFavoriteUserAssetTagsConfig.h
@@ -6,6 +6,8 @@
 #include "EditorConfigBase.h"
 #include "LocalFavoriteUserAssetTagsConfig.generated.h"
 
+#define UE_API USERASSETTAGSEDITOR_API
+
 USTRUCT()
 struct FPerTypeFavoriteUserAssetTags
 {
@@ -17,21 +19,21 @@ struct FPerTypeFavoriteUserAssetTags
 /**
  * 
  */
-UCLASS(EditorConfig="FavoriteUserAssetTags")
-class USERASSETTAGSEDITOR_API ULocalFavoriteUserAssetTagsConfig : public UEditorConfigBase
+UCLASS(MinimalAPI, EditorConfig="FavoriteUserAssetTags")
+class ULocalFavoriteUserAssetTagsConfig : public UEditorConfigBase
 {
 	GENERATED_BODY()
 
 public:
-	static ULocalFavoriteUserAssetTagsConfig* Get();
-	static void Shutdown();
+	static UE_API ULocalFavoriteUserAssetTagsConfig* Get();
+	static UE_API void Shutdown();
 	
 	TMap<FSoftClassPath, FPerTypeFavoriteUserAssetTags>& GetFavoriteUserAssetTagsMutable() { return FavoriteUserAssetTagsPerClass; }
 	
 	const TMap<FSoftClassPath, FPerTypeFavoriteUserAssetTags>& GetFavoriteUserAssetTags() const { return FavoriteUserAssetTagsPerClass; }
 
-	TSet<FName> GetFavoriteUserAssetTagsForClass(const UClass* Class);
-	void ToggleFavoriteUserAssetTag(const UClass* Class, FName InUserAssetTag);
+	UE_API TSet<FName> GetFavoriteUserAssetTagsForClass(const UClass* Class);
+	UE_API void ToggleFavoriteUserAssetTag(const UClass* Class, FName InUserAssetTag);
 private:
 	UPROPERTY(meta=(EditorConfig))
 	TMap<FSoftClassPath, FPerTypeFavoriteUserAssetTags> FavoriteUserAssetTagsPerClass;
@@ -39,5 +41,7 @@ private:
 	UPROPERTY(meta=(EditorConfig))
 	int32 MaxRecentUserAssetTags = 10;
 		
-	static TStrongObjectPtr<ULocalFavoriteUserAssetTagsConfig> Instance;
+	static UE_API TStrongObjectPtr<ULocalFavoriteUserAssetTagsConfig> Instance;
 };
+
+#undef UE_API
diff --git a/Engine/Source/Editor/UserAssetTagsEditor/Public/Config/TaggedAssetBrowserConfig.h b/Engine/Source/Editor/UserAssetTagsEditor/Public/Config/TaggedAssetBrowserConfig.h
--- a/Engine/Source/Editor/UserAssetTagsEditor/Public/Config/TaggedAssetBrowserConfig.h
+++ b/Engine/Source/Editor/UserAssetTagsEditor/Public/Config/TaggedAssetBrowserConfig.h
@@ -6,6 +6,8 @@
 #include "EditorConfigBase.h"
 #include "TaggedAssetBrowserConfig.generated.h"
 
+#define UE_API USERASSETTAGSEDITOR_API
+
 USTRUCT()
 struct FPerTaggedAssetBrowserSavedState
 {
@@ -15,15 +17,15 @@ struct FPerTaggedAssetBrowserSavedState
 	TArray<FName> PrimaryFilterSelection;
 };
 
-UCLASS(EditorConfig="TaggedAssetBrowser")
-class USERASSETTAGSEDITOR_API UTaggedAssetBrowserConfig : public UEditorConfigBase
+UCLASS(MinimalAPI, EditorConfig="TaggedAssetBrowser")
+class UTaggedAssetBrowserConfig : public UEditorConfigBase
 {
 	GENERATED_BODY()
 public:
 	DECLARE_MULTICAST_DELEGATE_OneParam(FOnPropertyChanged, const FPropertyChangedEvent&);
 		
-	static UTaggedAssetBrowserConfig* Get();
-	static void Shutdown();
+	static UE_API UTaggedAssetBrowserConfig* Get();
+	static UE_API void Shutdown();
 	
 	FOnPropertyChanged& OnPropertyChanged() { return OnPropertyChangedDelegate; }
 	
@@ -36,9 +38,11 @@ public:
 	UPROPERTY(meta=(EditorConfig))
 	bool bShowDeprecatedAssets = false;
 
-	virtual void PostEditChangeProperty(FPropertyChangedEvent& PropertyChangedEvent) override;	
+	UE_API virtual void PostEditChangeProperty(FPropertyChangedEvent& PropertyChangedEvent) override;	
 private:
-	static TStrongObjectPtr<UTaggedAssetBrowserConfig> Instance;
+	static UE_API TStrongObjectPtr<UTaggedAssetBrowserConfig> Instance;
 
 	FOnPropertyChanged OnPropertyChangedDelegate;
 };
+
+#undef UE_API
diff --git a/Engine/Source/Editor/UserAssetTagsEditor/Public/Config/UserAssetTagsEditorConfig.h b/Engine/Source/Editor/UserAssetTagsEditor/Public/Config/UserAssetTagsEditorConfig.h
--- a/Engine/Source/Editor/UserAssetTagsEditor/Public/Config/UserAssetTagsEditorConfig.h
+++ b/Engine/Source/Editor/UserAssetTagsEditor/Public/Config/UserAssetTagsEditorConfig.h
@@ -8,6 +8,8 @@
 #include "UObject/StrongObjectPtrTemplates.h"
 #include "UserAssetTagsEditorConfig.generated.h"
 
+#define UE_API USERASSETTAGSEDITOR_API
+
 USTRUCT()
 struct FPerUserAssetTagProviderViewOptions
 {
@@ -32,23 +34,23 @@ struct FUserAssetTagProviderViewOptions
 /**
  * 
  */
-UCLASS(EditorConfig=UserAssetTags)
-class USERASSETTAGSEDITOR_API UUserAssetTagsEditorConfig : public UEditorConfigBase
+UCLASS(MinimalAPI, EditorConfig=UserAssetTags)
+class UUserAssetTagsEditorConfig : public UEditorConfigBase
 {
 	GENERATED_BODY()
 
 public:
-	void ToggleSortByAlphabet();
-	bool ShouldSortByAlphabet() const;
+	UE_API void ToggleSortByAlphabet();
+	UE_API bool ShouldSortByAlphabet() const;
 	
-	bool IsProviderEnabled(const UClass* ProviderClass) const;
-	void ToggleProviderEnabled(const UClass* ProviderClass);
+	UE_API bool IsProviderEnabled(const UClass* ProviderClass) const;
+	UE_API void ToggleProviderEnabled(const UClass* ProviderClass);
 	
-	EUserAssetTagProviderMenuType GetProviderMenuType(const UClass* ProviderClass) const;
-	void SetProviderMenuType(const UClass* ProviderClass, EUserAssetTagProviderMenuType InMenuType);
+	UE_API EUserAssetTagProviderMenuType GetProviderMenuType(const UClass* ProviderClass) const;
+	UE_API void SetProviderMenuType(const UClass* ProviderClass, EUserAssetTagProviderMenuType InMenuType);
 
-	static UUserAssetTagsEditorConfig* Get();
-	static void Shutdown();
+	static UE_API UUserAssetTagsEditorConfig* Get();
+	static UE_API void Shutdown();
 private:
 	UPROPERTY(meta=(EditorConfig))
 	bool bSortByAlphabet = false;
@@ -56,5 +58,7 @@ private:
 	UPROPERTY(meta=(EditorConfig))
 	FUserAssetTagProviderViewOptions ProviderViewOptions;
 private:
-    static TStrongObjectPtr<UUserAssetTagsEditorConfig> Instance;
+    static UE_API TStrongObjectPtr<UUserAssetTagsEditorConfig> Instance;
 };
+
+#undef UE_API
diff --git a/Engine/Source/Editor/UserAssetTagsEditor/Public/DataHierarchyEditor/TaggedAssetBrowserConfigurationHierarchyViewModel.h b/Engine/Source/Editor/UserAssetTagsEditor/Public/DataHierarchyEditor/TaggedAssetBrowserConfigurationHierarchyViewModel.h
--- a/Engine/Source/Editor/UserAssetTagsEditor/Public/DataHierarchyEditor/TaggedAssetBrowserConfigurationHierarchyViewModel.h
+++ b/Engine/Source/Editor/UserAssetTagsEditor/Public/DataHierarchyEditor/TaggedAssetBrowserConfigurationHierarchyViewModel.h
@@ -6,6 +6,8 @@
 #include "DataHierarchyViewModelBase.h"
 #include "TaggedAssetBrowserConfigurationHierarchyViewModel.generated.h"
 
+#define UE_API USERASSETTAGSEDITOR_API
+
 class UTaggedAssetBrowserSection;
 class UTaggedAssetBrowserFilterBase;
 class UTaggedAssetBrowserFilter_All;
@@ -16,41 +18,41 @@ class UTaggedAssetBrowserConfiguration;
 /**
  * The view model for Data Hierarchy Editor integration to author filter views for a Tagged Asset Browser.
  */
-UCLASS()
-class USERASSETTAGSEDITOR_API UTaggedAssetBrowserConfigurationHierarchyViewModel : public UDataHierarchyViewModelBase
+UCLASS(MinimalAPI)
+class UTaggedAssetBrowserConfigurationHierarchyViewModel : public UDataHierarchyViewModelBase
 {
 	GENERATED_BODY()
 
 public:
-	void Initialize(UTaggedAssetBrowserConfiguration& InAsset);
+	UE_API void Initialize(UTaggedAssetBrowserConfiguration& InAsset);
 	UTaggedAssetBrowserConfiguration* GetConfigurationAsset() const { return ConfigurationAsset.Get(); }
 	
-	virtual UHierarchyRoot* GetHierarchyRoot() const override;
-	virtual TSharedPtr<FHierarchyElementViewModel> CreateCustomViewModelForElement(UHierarchyElement* Element, TSharedPtr<FHierarchyElementViewModel> Parent) override;
+	UE_API virtual UHierarchyRoot* GetHierarchyRoot() const override;
+	UE_API virtual TSharedPtr<FHierarchyElementViewModel> CreateCustomViewModelForElement(UHierarchyElement* Element, TSharedPtr<FHierarchyElementViewModel> Parent) override;
 
 	/** We don't need a sources panel or root as we only add new elements. */
 	virtual bool SupportsSourcePanel() const override { return false; }
 
 	/** We specify our custom section class */
-	virtual TSubclassOf<UHierarchySection> GetSectionDataClass() const override;
+	UE_API virtual TSubclassOf<UHierarchySection> GetSectionDataClass() const override;
 
 	/** We disables categories */
 	virtual TSubclassOf<UHierarchyCategory> GetCategoryDataClass() const override { return nullptr; }
 
 	/** The types to add in the UI are specified here. We differentiate between standalone and extension elements. */
-	virtual TArray<TSubclassOf<UHierarchyElement>> GetAdditionalTypesToAddInUi() const override;
+	UE_API virtual TArray<TSubclassOf<UHierarchyElement>> GetAdditionalTypesToAddInUi() const override;
 
 	/** If we select nothing, we select the root. The custom root view model points at the asset for editing in the details panel. */
 	virtual bool SelectRootInsteadOfNone() const override { return true; }
 private:
-	TArray<TSubclassOf<UTaggedAssetBrowserFilterBase>> GetAdditionalTypesToAddInUi_Standalone() const;
-	TArray<TSubclassOf<UTaggedAssetBrowserFilterBase>> GetAdditionalTypesToAddInUi_Extension() const;
+	UE_API TArray<TSubclassOf<UTaggedAssetBrowserFilterBase>> GetAdditionalTypesToAddInUi_Standalone() const;
+	UE_API TArray<TSubclassOf<UTaggedAssetBrowserFilterBase>> GetAdditionalTypesToAddInUi_Extension() const;
 
-	void EmptyCache(UObject* Object, FPropertyChangedEvent& PropertyChangedEvent);
+	UE_API void EmptyCache(UObject* Object, FPropertyChangedEvent& PropertyChangedEvent);
 
 protected:
-	virtual void InitializeInternal() override;
-	virtual void FinalizeInternal() override;
+	UE_API virtual void InitializeInternal() override;
+	UE_API virtual void FinalizeInternal() override;
 protected:
 	TWeakObjectPtr<UTaggedAssetBrowserConfiguration> ConfigurationAsset;
 	mutable TOptional<TArray<TSubclassOf<UHierarchyElement>>> AdditionalTypesCache;
@@ -114,3 +116,5 @@ namespace UE::UserAssetTags::ViewModels
 		virtual void RenameInternal(FName NewName) override;
 	};
 }
+
+#undef UE_API
diff --git a/Engine/Source/Editor/UserAssetTagsEditor/Public/TaggedAssetBrowserFilters/TaggedAssetBrowser_CommonFilters.h b/Engine/Source/Editor/UserAssetTagsEditor/Public/TaggedAssetBrowserFilters/TaggedAssetBrowser_CommonFilters.h
--- a/Engine/Source/Editor/UserAssetTagsEditor/Public/TaggedAssetBrowserFilters/TaggedAssetBrowser_CommonFilters.h
+++ b/Engine/Source/Editor/UserAssetTagsEditor/Public/TaggedAssetBrowserFilters/TaggedAssetBrowser_CommonFilters.h
@@ -8,6 +8,8 @@
 #include "MRUFavoritesList.h"
 #include "TaggedAssetBrowser_CommonFilters.generated.h"
 
+#define UE_API USERASSETTAGSEDITOR_API
+
 class UTaggedAssetBrowserSection;
 
 UCLASS()
@@ -56,26 +58,26 @@ public:
 };
 
 /** Show all assets. */
-UCLASS(DisplayName="All")
-class USERASSETTAGSEDITOR_API UTaggedAssetBrowserFilter_All : public UTaggedAssetBrowserFilterBase
+UCLASS(MinimalAPI, DisplayName="All")
+class UTaggedAssetBrowserFilter_All : public UTaggedAssetBrowserFilterBase
 {
 	GENERATED_BODY()
 };
 
 /** A filter that returns assets that have the specified User Asset Tag. Nested User Asset Tag filters will be combined. */
-UCLASS(DisplayName="User Asset Tag")
-class USERASSETTAGSEDITOR_API UTaggedAssetBrowserFilter_UserAssetTag : public UTaggedAssetBrowserFilterBase
+UCLASS(MinimalAPI, DisplayName="User Asset Tag")
+class UTaggedAssetBrowserFilter_UserAssetTag : public UTaggedAssetBrowserFilterBase
 {
 	GENERATED_BODY()
 public:
-	virtual FString ToString() const override;
+	UE_API virtual FString ToString() const override;
 	virtual TOptional<FString> GetInstanceIdentifier() const override { return ToString(); }
-	virtual FSlateIcon GetIcon() const override;
-	virtual void ModifyARFilterInternal(FARFilter& Filter) const override;
-	virtual bool ShouldFilterAssetInternal(const FAssetData& InAssetData) const override;
+	UE_API virtual FSlateIcon GetIcon() const override;
+	UE_API virtual void ModifyARFilterInternal(FARFilter& Filter) const override;
+	UE_API virtual bool ShouldFilterAssetInternal(const FAssetData& InAssetData) const override;
 
-	void SetUserAssetTag(FName InUserAssetTag);
-	bool DoesAssetHaveTag(const FAssetData& AssetCandidate) const;
+	UE_API void SetUserAssetTag(FName InUserAssetTag);
+	UE_API bool DoesAssetHaveTag(const FAssetData& AssetCandidate) const;
 
 protected:
 	UPROPERTY(EditAnywhere, Category="Filter")
@@ -83,23 +85,23 @@ protected:
 };
 
 /** A collection of multiple User Asset Tags. Will show all assets that own at least one of the contained tags. */
-UCLASS(DisplayName="User Asset Tag Collection")
-class USERASSETTAGSEDITOR_API UTaggedAssetBrowserFilter_UserAssetTagCollection : public UTaggedAssetBrowserFilterBase
+UCLASS(MinimalAPI, DisplayName="User Asset Tag Collection")
+class UTaggedAssetBrowserFilter_UserAssetTagCollection : public UTaggedAssetBrowserFilterBase
 {
 	GENERATED_BODY()
 
 public:
 
-	virtual FString ToString() const override;
+	UE_API virtual FString ToString() const override;
 	virtual TOptional<FString> GetInstanceIdentifier() const override { return ToString(); }
-	virtual FText GetTooltip() const override;
-	virtual FSlateIcon GetIcon() const override;
-	virtual void CreateAdditionalWidgets(TSharedPtr<SHorizontalBox> ExtensionBox) override;
+	UE_API virtual FText GetTooltip() const override;
+	UE_API virtual FSlateIcon GetIcon() const override;
+	UE_API virtual void CreateAdditionalWidgets(TSharedPtr<SHorizontalBox> ExtensionBox) override;
 
-	virtual void ModifyARFilterInternal(FARFilter& Filter) const override;
-	virtual bool ShouldFilterAssetInternal(const FAssetData& InAssetData) const override;
+	UE_API virtual void ModifyARFilterInternal(FARFilter& Filter) const override;
+	UE_API virtual bool ShouldFilterAssetInternal(const FAssetData& InAssetData) const override;
 	
-	void SetCollectionName(FName InName);
+	UE_API void SetCollectionName(FName InName);
 	
 protected:
 	UPROPERTY(EditAnywhere, Category="Filter")
@@ -109,41 +111,41 @@ protected:
 	FText Description;
 	
 private:
-	FText GetContainedChildrenNumberText() const;
-	EVisibility GetContainedChildrenNumberTextVisibility() const;
-	FText GetContainedChildrenNumberTooltip() const;
+	UE_API FText GetContainedChildrenNumberText() const;
+	UE_API EVisibility GetContainedChildrenNumberTextVisibility() const;
+	UE_API FText GetContainedChildrenNumberTooltip() const;
 };
 
 /** Filter for recently used assets. */
-UCLASS(DisplayName="Recent")
-class USERASSETTAGSEDITOR_API UTaggedAssetBrowserFilter_Recent : public UTaggedAssetBrowserFilterBase
+UCLASS(MinimalAPI, DisplayName="Recent")
+class UTaggedAssetBrowserFilter_Recent : public UTaggedAssetBrowserFilterBase
 {
 	GENERATED_BODY()
 
 public:
 
-	virtual void InitializeInternal(const FTaggedAssetBrowserContext& InContext) override;
-	virtual bool ShouldFilterAssetInternal(const FAssetData& InAssetData) const override;
-	virtual FSlateIcon GetIcon() const override;
+	UE_API virtual void InitializeInternal(const FTaggedAssetBrowserContext& InContext) override;
+	UE_API virtual bool ShouldFilterAssetInternal(const FAssetData& InAssetData) const override;
+	UE_API virtual FSlateIcon GetIcon() const override;
 
-	bool IsAssetRecent(const FAssetData& AssetCandidate) const;
+	UE_API bool IsAssetRecent(const FAssetData& AssetCandidate) const;
 
 	TAttribute<const FMainMRUFavoritesList*> ListAttribute;
 };
 
 /** Filter for assets based on directories. */
-UCLASS(DisplayName="Directories")
-class USERASSETTAGSEDITOR_API UTaggedAssetBrowserFilter_Directories : public UTaggedAssetBrowserFilterBase
+UCLASS(MinimalAPI, DisplayName="Directories")
+class UTaggedAssetBrowserFilter_Directories : public UTaggedAssetBrowserFilterBase
 {
 	GENERATED_BODY()
 
 public:
 
-	virtual FString ToString() const override;
+	UE_API virtual FString ToString() const override;
 	virtual TOptional<FString> GetInstanceIdentifier() const override { return ToString(); }
-	virtual void ModifyARFilterInternal(FARFilter& Filter) const override;
-	virtual FText GetTooltip() const override;
-	virtual FSlateIcon GetIcon() const override;
+	UE_API virtual void ModifyARFilterInternal(FARFilter& Filter) const override;
+	UE_API virtual FText GetTooltip() const override;
+	UE_API virtual FSlateIcon GetIcon() const override;
 
 	UPROPERTY(EditAnywhere, Category="Filter", meta=(ContentDir))
 	TArray<FDirectoryPath> DirectoryPaths;
@@ -154,19 +156,19 @@ public:
 };
 
 /** Filter for assets based on class. By default will use first class as an icon. */
-UCLASS(DisplayName="Classes")
-class USERASSETTAGSEDITOR_API UTaggedAssetBrowserFilter_Class : public UTaggedAssetBrowserFilterBase
+UCLASS(MinimalAPI, DisplayName="Classes")
+class UTaggedAssetBrowserFilter_Class : public UTaggedAssetBrowserFilterBase
 {
 	GENERATED_BODY()
 
 public:
 
-	virtual FString ToString() const override;
+	UE_API virtual FString ToString() const override;
 	virtual TOptional<FString> GetInstanceIdentifier() const override { return ToString(); }
-	virtual FText GetTooltip() const override;
-	virtual FSlateIcon GetIcon() const override;
+	UE_API virtual FText GetTooltip() const override;
+	UE_API virtual FSlateIcon GetIcon() const override;
 	
-	virtual void ModifyARFilterInternal(FARFilter& Filter) const override;
+	UE_API virtual void ModifyARFilterInternal(FARFilter& Filter) const override;
 
 	UPROPERTY(EditAnywhere, Category="Filter", meta=(AllowAbstract))
 	TArray<TObjectPtr<UClass>> Classes;
@@ -179,3 +181,5 @@ public:
 	UPROPERTY(EditAnywhere, Category = "Filter")
 	FName FilterName;
 };
+
+#undef UE_API
diff --git a/Engine/Source/Editor/UserAssetTagsEditor/Public/TaggedAssetBrowserMenuFilters.h b/Engine/Source/Editor/UserAssetTagsEditor/Public/TaggedAssetBrowserMenuFilters.h
--- a/Engine/Source/Editor/UserAssetTagsEditor/Public/TaggedAssetBrowserMenuFilters.h
+++ b/Engine/Source/Editor/UserAssetTagsEditor/Public/TaggedAssetBrowserMenuFilters.h
@@ -6,6 +6,8 @@
 #include "FrontendFilterBase.h"
 #include "TaggedAssetBrowserMenuFilters.generated.h"
 
+#define UE_API USERASSETTAGSEDITOR_API
+
 /** A helper struct passed into filters for initialization, since some filters require contextual knowledge to filter correctly. */
 struct FTaggedAssetBrowserContext
 {
@@ -24,38 +26,38 @@ struct FTaggedAssetBrowserContext
  *
  *  While filters are generally part of the hierarchy, sections can also specify additional filters
  */
-UCLASS(Abstract, EditInlineNew)
-class USERASSETTAGSEDITOR_API UTaggedAssetBrowserFilterBase : public UHierarchyItem
+UCLASS(MinimalAPI, Abstract, EditInlineNew)
+class UTaggedAssetBrowserFilterBase : public UHierarchyItem
 {
 	GENERATED_BODY()
 	
 public:	
 	/** The filter has a chance to initialize itself with contextual data to actually work. */
-	void Initialize(const FTaggedAssetBrowserContext& InContext);
+	UE_API void Initialize(const FTaggedAssetBrowserContext& InContext);
 
-	FName GetIdentifier() const;
+	UE_API FName GetIdentifier() const;
 
 	/** The string to display for this filter. */
 	virtual FString ToString() const override { return GetClass()->GetDisplayNameText().ToString(); }
 	/** The tooltip to display for this filter. */
 	virtual FText GetTooltip() const { return GetClass()->GetToolTipText(); }
 	/** The optional icon to display for this filter. */
-	virtual FSlateIcon GetIcon() const;
+	UE_API virtual FSlateIcon GetIcon() const;
 	/** The optional widgets to display for this filter's row. */
 	virtual void CreateAdditionalWidgets(TSharedPtr<SHorizontalBox> ExtensionBox) { }
 	
-	const FSlateBrush* GetIconBrush() const;
+	UE_API const FSlateBrush* GetIconBrush() const;
 
 	/** How this filter should modify the AR filter, when selected. */
-	void ModifyARFilter(FARFilter& Filter) const;
+	UE_API void ModifyARFilter(FARFilter& Filter) const;
 	/** An additional check for each asset that passes the AR filter. */
-	bool ShouldFilterAsset(const FAssetData& InAssetData) const;
+	UE_API bool ShouldFilterAsset(const FAssetData& InAssetData) const;
 	
 	/** Used for search purposes. Override this to add your own implementation for your custom filter. */
-	virtual bool DoesFilterMatchTextQuery(const FText& Text);
+	UE_API virtual bool DoesFilterMatchTextQuery(const FText& Text);
 
 	/** If this filter is selected in the Tagged Asset Browser. */
-	bool IsSelectedFilter() const;
+	UE_API bool IsSelectedFilter() const;
 private:
 	/** Implementation of your own initialization when the Tagged Asset Browser starts using this filter. */
 	virtual void InitializeInternal(const FTaggedAssetBrowserContext& InContext) { }
@@ -73,4 +75,6 @@ protected:
 	TOptional<FTaggedAssetBrowserContext> ActiveContext;
 };
 
-uint32 GetTypeHash(const UTaggedAssetBrowserFilterBase& Filter);
\ No newline at end of file
+uint32 GetTypeHash(const UTaggedAssetBrowserFilterBase& Filter);
+
+#undef UE_API
diff --git a/Engine/Source/Editor/UserAssetTagsEditor/Public/UserAssetTagProvider.h b/Engine/Source/Editor/UserAssetTagsEditor/Public/UserAssetTagProvider.h
--- a/Engine/Source/Editor/UserAssetTagsEditor/Public/UserAssetTagProvider.h
+++ b/Engine/Source/Editor/UserAssetTagsEditor/Public/UserAssetTagProvider.h
@@ -6,6 +6,8 @@
 #include "Templates/SharedPointer.h"
 #include "UserAssetTagProvider.generated.h"
 
+#define UE_API USERASSETTAGSEDITOR_API
+
 class SHorizontalBox;
 class UUserAssetTagEditorContext;
 class SWidget;
@@ -23,13 +25,13 @@ enum class EUserAssetTagProviderMenuType
  *  A Tag Provider can override functions to add custom UI elements to the Tag Editor that help with managing tags.
  *  E.g.: A 'Favorites' tag provider can add a 'favorite/unfavorite' toggle to individual tag rows.
  */
-UCLASS(Transient, Abstract)
-class USERASSETTAGSEDITOR_API UUserAssetTagProvider : public UObject
+UCLASS(MinimalAPI, Transient, Abstract)
+class UUserAssetTagProvider : public UObject
 {
 	GENERATED_BODY()
 
 public:
-	struct USERASSETTAGSEDITOR_API FResultWithUserFeedback
+	struct FResultWithUserFeedback
 	{
 		FResultWithUserFeedback(bool bInCanPerform) : bResult(bInCanPerform) {}
 		
@@ -56,9 +58,9 @@ public:
 	};
 
 	/** By default we return the class display name text. */
-	virtual FText GetDisplayNameText(const UUserAssetTagEditorContext* Context) const;
+	UE_API virtual FText GetDisplayNameText(const UUserAssetTagEditorContext* Context) const;
 	/** By default we return the class tooltip text. */
-	virtual FText GetToolTipText(const UUserAssetTagEditorContext* Context) const;
+	UE_API virtual FText GetToolTipText(const UUserAssetTagEditorContext* Context) const;
 	
 	/** Returns a list of suggested user asset tags. */
 	virtual TSet<FName> GetSuggestedUserAssetTags(const UUserAssetTagEditorContext* Context) const PURE_VIRTUAL(UUserAssetTagProvider::GetSuggestedUserAssetTags, return {}; )
@@ -82,3 +84,5 @@ public:
 	/** Lets you add additional */
 	virtual void AddToolbarMenuEntries(class UToolMenu* DynamicMenu, const UUserAssetTagEditorContext* Context) const { }
 };
+
+#undef UE_API
diff --git a/Engine/Source/Editor/UserAssetTagsEditor/Public/UserAssetTagsEditorModule.h b/Engine/Source/Editor/UserAssetTagsEditor/Public/UserAssetTagsEditorModule.h
--- a/Engine/Source/Editor/UserAssetTagsEditor/Public/UserAssetTagsEditorModule.h
+++ b/Engine/Source/Editor/UserAssetTagsEditor/Public/UserAssetTagsEditorModule.h
@@ -6,6 +6,8 @@
 #include "Modules/ModuleManager.h"
 #include "Delegates/DelegateCombinations.h"
 
+#define UE_API USERASSETTAGSEDITOR_API
+
 struct FHierarchyElementViewModel;
 class UTaggedAssetBrowserFilterBase;
 class SDockTab;
@@ -14,36 +16,36 @@ class FAssetRegistryTagsContext;
 
 USERASSETTAGSEDITOR_API DECLARE_LOG_CATEGORY_EXTERN(LogUserAssetTags, Log, All);
 
-class USERASSETTAGSEDITOR_API FUserAssetTagsEditorModule : public IModuleInterface
+class FUserAssetTagsEditorModule : public IModuleInterface
 {
 public:
 	DECLARE_DELEGATE_RetVal_ThreeParams(TSharedPtr<FHierarchyElementViewModel>, FOnGetViewModelFactory, UHierarchyElement* Element, TSharedPtr<FHierarchyElementViewModel> InParent, TWeakObjectPtr<UDataHierarchyViewModelBase> InViewModel);
 
-	FUserAssetTagsEditorModule();
-	virtual ~FUserAssetTagsEditorModule() override;
+	UE_API FUserAssetTagsEditorModule();
+	UE_API virtual ~FUserAssetTagsEditorModule() override;
 	
-	static FUserAssetTagsEditorModule& Get();
+	static UE_API FUserAssetTagsEditorModule& Get();
 	
-	virtual void StartupModule() override;
-	virtual void ShutdownModule() override;
+	UE_API virtual void StartupModule() override;
+	UE_API virtual void ShutdownModule() override;
 
 	/** To register a custom view model for display in the Tagged Asset Browser. */
-	void RegisterConfigurationHierarchyElementViewModel(TSubclassOf<UHierarchyElement> FilterType, FOnGetViewModelFactory InFactory);
-	void UnregisterConfigurationHierarchyElementViewModel(TSubclassOf<UHierarchyElement> FilterType);
+	UE_API void RegisterConfigurationHierarchyElementViewModel(TSubclassOf<UHierarchyElement> FilterType, FOnGetViewModelFactory InFactory);
+	UE_API void UnregisterConfigurationHierarchyElementViewModel(TSubclassOf<UHierarchyElement> FilterType);
 	const TMap<TSubclassOf<UHierarchyElement>, FOnGetViewModelFactory>& GetConfigurationHierarchyViewModelFactories() const { return ConfigurationHierarchyViewModelFactories; }
 
-	static FName ManageTagsTabId;
+	static UE_API FName ManageTagsTabId;
 private:
-	void RegisterDefaultConfigurationHierarchyElementViewModels();
-	void RegisterDetailsCustomizations();
+	UE_API void RegisterDefaultConfigurationHierarchyElementViewModels();
+	UE_API void RegisterDetailsCustomizations();
 
 	/** We extend the Content Browser command list so that the shortcut works. */
-	void OnRegisterCommandList(FName ContextName, TSharedRef<FUICommandList> CommandList);
+	UE_API void OnRegisterCommandList(FName ContextName, TSharedRef<FUICommandList> CommandList);
 
-	static TSharedRef<SDockTab> SpawnUserAssetTagsEditorNomadTab(const FSpawnTabArgs& SpawnTabArgs);
+	static UE_API TSharedRef<SDockTab> SpawnUserAssetTagsEditorNomadTab(const FSpawnTabArgs& SpawnTabArgs);
 
-	static void GetAssetRegistryTagsFromUserAssetTags(FAssetRegistryTagsContext AssetRegistryTagsContext);
-	static void TransferUserAssetTagsFromOldPackageToNewPackage(const FMetaData& OldMetaData, FMetaData& NewMetaData, const UPackage* OldPackage, UPackage* NewPackage);
+	static UE_API void GetAssetRegistryTagsFromUserAssetTags(FAssetRegistryTagsContext AssetRegistryTagsContext);
+	static UE_API void TransferUserAssetTagsFromOldPackageToNewPackage(const FMetaData& OldMetaData, FMetaData& NewMetaData, const UPackage* OldPackage, UPackage* NewPackage);
 
 private:
 	TMap<TSubclassOf<UHierarchyElement>, FOnGetViewModelFactory> ConfigurationHierarchyViewModelFactories;
@@ -51,3 +53,5 @@ private:
 	FDelegateHandle GetAssetRegistryTagsFromUserAssetTagsDelegateHandle;
 	FDelegateHandle TransferUserAssetTagsFromOldPackageToNewPackageDelegateHandle;
 };
+
+#undef UE_API
diff --git a/Engine/Source/Editor/UserAssetTagsEditor/Public/Widgets/STaggedAssetBrowserContent.h b/Engine/Source/Editor/UserAssetTagsEditor/Public/Widgets/STaggedAssetBrowserContent.h
--- a/Engine/Source/Editor/UserAssetTagsEditor/Public/Widgets/STaggedAssetBrowserContent.h
+++ b/Engine/Source/Editor/UserAssetTagsEditor/Public/Widgets/STaggedAssetBrowserContent.h
@@ -7,7 +7,9 @@
 #include "Filters/SAssetFilterBar.h"
 #include "Widgets/SCompoundWidget.h"
 
-class USERASSETTAGSEDITOR_API STaggedAssetBrowserContent : public SCompoundWidget
+#define UE_API USERASSETTAGSEDITOR_API
+
+class STaggedAssetBrowserContent : public SCompoundWidget
 {
 public:
 	DECLARE_DELEGATE_RetVal_OneParam(bool, FOnShouldFilterAsset, const FAssetData&)
@@ -20,11 +22,11 @@ public:
 	SLATE_END_ARGS()
 
 	/** Constructs this widget with InArgs */
-	void Construct(const FArguments& InArgs);
+	UE_API void Construct(const FArguments& InArgs);
 
-	void SetARFilter(FARFilter InFilter);
+	UE_API void SetARFilter(FARFilter InFilter);
 
-	TArray<FAssetData> GetCurrentSelection() const;
+	UE_API TArray<FAssetData> GetCurrentSelection() const;
 private:	
 	FRefreshAssetViewDelegate RefreshAssetViewDelegate;
 	FSyncToAssetsDelegate SyncToAssetsDelegate;
@@ -33,3 +35,5 @@ private:
 
 	FAssetPickerConfig InitialAssetPickerConfig;
 };
+
+#undef UE_API
diff --git a/Engine/Source/Editor/UserAssetTagsEditor/Public/Widgets/STaggedAssetBrowserSections.h b/Engine/Source/Editor/UserAssetTagsEditor/Public/Widgets/STaggedAssetBrowserSections.h
--- a/Engine/Source/Editor/UserAssetTagsEditor/Public/Widgets/STaggedAssetBrowserSections.h
+++ b/Engine/Source/Editor/UserAssetTagsEditor/Public/Widgets/STaggedAssetBrowserSections.h
@@ -5,6 +5,8 @@
 #include "TaggedAssetBrowserFilters/TaggedAssetBrowser_CommonFilters.h"
 #include "Widgets/SCompoundWidget.h"
 
+#define UE_API USERASSETTAGSEDITOR_API
+
 class UTaggedAssetBrowserSection;
 class UTaggedAssetBrowserFilterRoot;
 /**
@@ -40,7 +42,7 @@ private:
 	FOnGetContent OnGetMenuContentDelegate;
 };
 
-class USERASSETTAGSEDITOR_API STaggedAssetBrowserSections : public SCompoundWidget
+class STaggedAssetBrowserSections : public SCompoundWidget
 {
 public:
 	DECLARE_DELEGATE_OneParam(FOnSectionSelected, const UTaggedAssetBrowserSection* Section)
@@ -53,15 +55,15 @@ public:
 	SLATE_END_ARGS()
 
 	/** Constructs this widget with InArgs */
-	void Construct(const FArguments& InArgs, const UTaggedAssetBrowserFilterRoot& InFilterRoot);
+	UE_API void Construct(const FArguments& InArgs, const UTaggedAssetBrowserFilterRoot& InFilterRoot);
 
-	void RebuildWidget();
+	UE_API void RebuildWidget();
 
 	const UTaggedAssetBrowserSection* GetActiveSection() const { return ActiveSection.Get(); }
 
 private:
-	void OnSectionSelected(ECheckBoxState CheckBoxState, const UTaggedAssetBrowserSection* InSection);
-	ECheckBoxState IsSectionActive(const UTaggedAssetBrowserSection* InSection) const;
+	UE_API void OnSectionSelected(ECheckBoxState CheckBoxState, const UTaggedAssetBrowserSection* InSection);
+	UE_API ECheckBoxState IsSectionActive(const UTaggedAssetBrowserSection* InSection) const;
 
 private:
 	TWeakObjectPtr<const UTaggedAssetBrowserFilterRoot> FilterRoot;
@@ -72,3 +74,5 @@ private:
 	TSharedPtr<SScrollBox> ScrollBox;
 
 };
+
+#undef UE_API
diff --git a/Engine/Source/Editor/UserAssetTagsEditor/Public/Widgets/SUserAssetTagBrowserSelectedAssetDetails.h b/Engine/Source/Editor/UserAssetTagsEditor/Public/Widgets/SUserAssetTagBrowserSelectedAssetDetails.h
--- a/Engine/Source/Editor/UserAssetTagsEditor/Public/Widgets/SUserAssetTagBrowserSelectedAssetDetails.h
+++ b/Engine/Source/Editor/UserAssetTagsEditor/Public/Widgets/SUserAssetTagBrowserSelectedAssetDetails.h
@@ -6,6 +6,8 @@
 #include "AssetThumbnail.h"
 #include "AssetRegistry/AssetData.h"
 
+#define UE_API USERASSETTAGSEDITOR_API
+
 struct FDisplayedPropertyData
 {
 	DECLARE_DELEGATE_RetVal_OneParam(bool, FShouldDisplayProperty, const FAssetData&)
@@ -32,7 +34,7 @@ struct FTaggedAssetBrowserDetailsDisplayDatabase
 
 DECLARE_DELEGATE_OneParam(FOnAssetTagActivated, const FName& UserAssetTag);
 
-class USERASSETTAGSEDITOR_API SUserAssetTag : public SCompoundWidget
+class SUserAssetTag : public SCompoundWidget
 {
 	SLATE_BEGIN_ARGS(SUserAssetTag)
 	{
@@ -42,10 +44,10 @@ class USERASSETTAGSEDITOR_API SUserAssetTag : public SCompoundWidget
 	SLATE_END_ARGS()
 
 	// TODO (ME) Add Domain Info for global lookup for additional info such as tooltip, description etc
-	void Construct(const FArguments& InArgs, const FName& UserAssetTag);
+	UE_API void Construct(const FArguments& InArgs, const FName& UserAssetTag);
 
 private:
-	FReply OnClicked() const;
+	UE_API FReply OnClicked() const;
 
 private:
 	FName UserAssetTag;
@@ -53,7 +55,7 @@ private:
 	TOptional<FText> OnAssetTagActivatedTooltip;
 };
 
-class USERASSETTAGSEDITOR_API SUserAssetTagRow : public SCompoundWidget
+class SUserAssetTagRow : public SCompoundWidget
 {
 	SLATE_BEGIN_ARGS(SUserAssetTagRow)
 		{
@@ -62,14 +64,14 @@ class USERASSETTAGSEDITOR_API SUserAssetTagRow : public SCompoundWidget
 		SLATE_ARGUMENT(TOptional<FText>, OnAssetTagActivatedTooltip)
 	SLATE_END_ARGS()
 
-	void Construct(const FArguments& InArgs, const FAssetData& Asset);
+	UE_API void Construct(const FArguments& InArgs, const FAssetData& Asset);
 };
 
 /**
  * The default widget for the right hand side of the tagged asset browser, showing details of the selected assets.
  * It requires previous configuration using the FUserAssetTag
  */
-class USERASSETTAGSEDITOR_API SUserAssetTagBrowserSelectedAssetDetails : public SCompoundWidget
+class SUserAssetTagBrowserSelectedAssetDetails : public SCompoundWidget
 {
 public:
 	DECLARE_DELEGATE_RetVal_OneParam(TSharedRef<SWidget>, FOnGenerateWidget, const FAssetData& AssetData)
@@ -95,10 +97,10 @@ public:
 	SLATE_END_ARGS()
 	
 	/** Constructs this widget with InArgs */
-	void Construct(const FArguments& InArgs, const FAssetData& Asset);
+	UE_API void Construct(const FArguments& InArgs, const FAssetData& Asset);
 
 private:
-	FReply CopyAssetPathToClipboard() const;
+	UE_API FReply CopyAssetPathToClipboard() const;
 	
 private:
 	FAssetData AssetData;
@@ -108,11 +110,13 @@ private:
 	FOnAssetTagActivated OnAssetTagActivated;
 	TOptional<FText> OnAssetTagActivatedTooltip;
 
-	TSharedRef<SWidget> CreateAssetThumbnailWidget();
-	TSharedRef<SWidget> CreateTitleWidget();
-	TSharedRef<SWidget> CreateTypeWidget();
-	TSharedRef<SWidget> CreatePathWidget();
-	TSharedRef<SWidget> CreateDescriptionWidget();
-	TSharedRef<SWidget> CreateOptionalPropertiesList();
-	TSharedRef<SWidget> CreateAssetTagRow();
+	UE_API TSharedRef<SWidget> CreateAssetThumbnailWidget();
+	UE_API TSharedRef<SWidget> CreateTitleWidget();
+	UE_API TSharedRef<SWidget> CreateTypeWidget();
+	UE_API TSharedRef<SWidget> CreatePathWidget();
+	UE_API TSharedRef<SWidget> CreateDescriptionWidget();
+	UE_API TSharedRef<SWidget> CreateOptionalPropertiesList();
+	UE_API TSharedRef<SWidget> CreateAssetTagRow();
 };
+
+#undef UE_API
diff --git a/Engine/Source/Editor/UserAssetTagsEditor/Public/Widgets/SUserAssetTagsEditor.h b/Engine/Source/Editor/UserAssetTagsEditor/Public/Widgets/SUserAssetTagsEditor.h
--- a/Engine/Source/Editor/UserAssetTagsEditor/Public/Widgets/SUserAssetTagsEditor.h
+++ b/Engine/Source/Editor/UserAssetTagsEditor/Public/Widgets/SUserAssetTagsEditor.h
@@ -8,6 +8,8 @@
 #include "Templates/SharedPointer.h"
 #include "UserAssetTagEditorMenuContexts.h"
 
+#define UE_API USERASSETTAGSEDITOR_API
+
 struct FToolMenuContext;
 struct FToolUIAction;
 struct FToolMenuCustomWidgetContext;
@@ -20,7 +22,7 @@ struct FUserAssetTagInfo
 };
 
 /** The actual editor widget to assign/unassign tags. */
-class USERASSETTAGSEDITOR_API SUserAssetTagsEditor : public SCompoundWidget
+class SUserAssetTagsEditor : public SCompoundWidget
 {
 public:
 	SLATE_BEGIN_ARGS(SUserAssetTagsEditor)
@@ -30,64 +32,64 @@ public:
 	SLATE_END_ARGS()
 
 	/** Constructs this widget with InArgs */
-	void Construct(const FArguments& InArgs);
-	virtual ~SUserAssetTagsEditor() override;
+	UE_API void Construct(const FArguments& InArgs);
+	UE_API virtual ~SUserAssetTagsEditor() override;
 	
 	const TArray<TSharedPtr<FAssetData>>& GetSelectedAssets() const { return SelectedAssets; }
 
-	void RefreshDataAndMenus();
-	static void RefreshMenus();
-	static void RefreshSelectedAssetsMenu();
-	static void RefreshSuggestedTagMenus();
-	static void RefreshOwnedTagMenus();
-	static void RefreshProviderExtensionToolbarMenu();
+	UE_API void RefreshDataAndMenus();
+	static UE_API void RefreshMenus();
+	static UE_API void RefreshSelectedAssetsMenu();
+	static UE_API void RefreshSuggestedTagMenus();
+	static UE_API void RefreshOwnedTagMenus();
+	static UE_API void RefreshProviderExtensionToolbarMenu();
 
 private:
-	TSharedPtr<SWidget> CreateSelectedAssetList();
-	TSharedPtr<SWidget> CreateSuggestedTagList();
-	TSharedPtr<SWidget> CreateOwnedTagsList();
-	TSharedPtr<SWidget> CreateToolbar();
+	UE_API TSharedPtr<SWidget> CreateSelectedAssetList();
+	UE_API TSharedPtr<SWidget> CreateSuggestedTagList();
+	UE_API TSharedPtr<SWidget> CreateOwnedTagsList();
+	UE_API TSharedPtr<SWidget> CreateToolbar();
 
-	static TArray<FAssetData> GetCurrentContentBrowserAssetSelection();
+	static UE_API TArray<FAssetData> GetCurrentContentBrowserAssetSelection();
 	
-	void OnCommitNewTag(const FText& Text, ETextCommit::Type CommitType);
-	FReply OnAddTagButtonClicked();
-	bool IsAddTagButtonEnabled() const;
+	UE_API void OnCommitNewTag(const FText& Text, ETextCommit::Type CommitType);
+	UE_API FReply OnAddTagButtonClicked();
+	UE_API bool IsAddTagButtonEnabled() const;
 
-	const FSlateBrush* GetViewOptionsBadgeIcon() const;
+	UE_API const FSlateBrush* GetViewOptionsBadgeIcon() const;
 	
-	static TSharedRef<SWidget> GenerateRowContent_SelectedAsset(const FAssetData& AssetData, const UUserAssetTagEditorContext* Context);
-	static TSharedRef<SWidget> GenerateRowContent_SuggestedTag(FName UserAssetTag, const UUserAssetTagEditorContext* Context);
-	static TSharedRef<SWidget> GenerateRowContent_OwnedTag(FName UserAssetTag, const UUserAssetTagEditorContext* Context);
-	static TSharedRef<SWidget> CreateMenuControlWidget(const FToolMenuContext& ToolMenuContext, const FToolMenuCustomWidgetContext& ToolMenuCustomWidgetContext, const UClass* Class);
-
-	TSharedRef<SWidget> OnGetViewOptions();
-	static void ToggleSortByAlphabet(const FToolMenuContext& ToolMenuContext);
-	static ECheckBoxState GetShouldSortByAlphabet(const FToolMenuContext& ToolMenuContext);
-	static void ToggleViewOption_ProviderClass(const FToolMenuContext& ToolMenuContext, const UClass* ProviderClass);
-	static ECheckBoxState GetViewOptions_IsProviderClassEnabled(const FToolMenuContext& ToolMenuContext, const UClass* Class);
-
-	static void SetViewOption_ProviderClassMenuType(EUserAssetTagProviderMenuType InMenuType, const UClass* ProviderClass);
-	static EUserAssetTagProviderMenuType GetViewOptions_ProviderClassMenuType(const UClass* ProviderClass);
+	static UE_API TSharedRef<SWidget> GenerateRowContent_SelectedAsset(const FAssetData& AssetData, const UUserAssetTagEditorContext* Context);
+	static UE_API TSharedRef<SWidget> GenerateRowContent_SuggestedTag(FName UserAssetTag, const UUserAssetTagEditorContext* Context);
+	static UE_API TSharedRef<SWidget> GenerateRowContent_OwnedTag(FName UserAssetTag, const UUserAssetTagEditorContext* Context);
+	static UE_API TSharedRef<SWidget> CreateMenuControlWidget(const FToolMenuContext& ToolMenuContext, const FToolMenuCustomWidgetContext& ToolMenuCustomWidgetContext, const UClass* Class);
+
+	UE_API TSharedRef<SWidget> OnGetViewOptions();
+	static UE_API void ToggleSortByAlphabet(const FToolMenuContext& ToolMenuContext);
+	static UE_API ECheckBoxState GetShouldSortByAlphabet(const FToolMenuContext& ToolMenuContext);
+	static UE_API void ToggleViewOption_ProviderClass(const FToolMenuContext& ToolMenuContext, const UClass* ProviderClass);
+	static UE_API ECheckBoxState GetViewOptions_IsProviderClassEnabled(const FToolMenuContext& ToolMenuContext, const UClass* Class);
+
+	static UE_API void SetViewOption_ProviderClassMenuType(EUserAssetTagProviderMenuType InMenuType, const UClass* ProviderClass);
+	static UE_API EUserAssetTagProviderMenuType GetViewOptions_ProviderClassMenuType(const UClass* ProviderClass);
 	
-	static const TArray<const UUserAssetTagProvider*>& GetAllProviderCDOs();
-	static TArray<const UUserAssetTagProvider*> GetAllValidDefaultEnabledProviderCDOs(const UUserAssetTagEditorContext* Context);
-	static TArray<const UUserAssetTagProvider*> GetValidProviderCDOs(const UUserAssetTagEditorContext* Context);
-	static TArray<const UUserAssetTagProvider*> GetEnabledProviderCDOs(const UUserAssetTagEditorContext* Context);
+	static UE_API const TArray<const UUserAssetTagProvider*>& GetAllProviderCDOs();
+	static UE_API TArray<const UUserAssetTagProvider*> GetAllValidDefaultEnabledProviderCDOs(const UUserAssetTagEditorContext* Context);
+	static UE_API TArray<const UUserAssetTagProvider*> GetValidProviderCDOs(const UUserAssetTagEditorContext* Context);
+	static UE_API TArray<const UUserAssetTagProvider*> GetEnabledProviderCDOs(const UUserAssetTagEditorContext* Context);
 
-	static FToolUIAction CreateToggleTagCheckboxAction_SuggestedTag(const FName& InUserAssetTag);
-	static FToolUIAction CreateToggleTagCheckboxAction_OwnedTag(const FName& InUserAssetTag);
+	static UE_API FToolUIAction CreateToggleTagCheckboxAction_SuggestedTag(const FName& InUserAssetTag);
+	static UE_API FToolUIAction CreateToggleTagCheckboxAction_OwnedTag(const FName& InUserAssetTag);
 
-	virtual FReply OnKeyDown(const FGeometry& MyGeometry, const FKeyEvent& InKeyEvent) override;
+	UE_API virtual FReply OnKeyDown(const FGeometry& MyGeometry, const FKeyEvent& InKeyEvent) override;
 
-	static FReply OnDocumentationRequested();
-	static FText GetDocumentationButtonTooltipText();
+	static UE_API FReply OnDocumentationRequested();
+	static UE_API FText GetDocumentationButtonTooltipText();
 
-	static FReply OnDeleteTagClicked(FName InUserAssetTag, const UUserAssetTagEditorContext* Context);
+	static UE_API FReply OnDeleteTagClicked(FName InUserAssetTag, const UUserAssetTagEditorContext* Context);
 
-	void HandleAssetSelectionChanged(const TArray<FAssetData>& AssetData, bool bIsPrimaryBrowser);
+	UE_API void HandleAssetSelectionChanged(const TArray<FAssetData>& AssetData, bool bIsPrimaryBrowser);
 	
-	static TArray<FAssetData> TransformAssetData(const TArray<TSharedPtr<FAssetData>>& InAssetData);
+	static UE_API TArray<FAssetData> TransformAssetData(const TArray<TSharedPtr<FAssetData>>& InAssetData);
 private:
 	TSharedPtr<SWidget> SelectedAssetList;
 	TSharedPtr<SWidget> TagSuggestionList;
@@ -102,5 +104,7 @@ private:
 	TStrongObjectPtr<UUserAssetTagEditorContext> ThisContext;
 
 private:
-	static TArray<const UUserAssetTagProvider*> CachedProviderCDOs;
+	static UE_API TArray<const UUserAssetTagProvider*> CachedProviderCDOs;
 };
+
+#undef UE_API
