Removed TransformAssetData calls
Changed User Asset Tag Editor property from SharedPtr struct array to regular struct array

This helps with perf when selecting a large amount of assets

#Synced-CL 49558609

diff --git a/Engine/Source/Editor/UserAssetTagsEditor/Private/UserAssetTagProviders/UserAssetTagProvider_LocalAssetTypeFavorites.cpp b/Engine/Source/Editor/UserAssetTagsEditor/Private/UserAssetTagProviders/UserAssetTagProvider_LocalAssetTypeFavorites.cpp
--- a/Engine/Source/Editor/UserAssetTagsEditor/Private/UserAssetTagProviders/UserAssetTagProvider_LocalAssetTypeFavorites.cpp
+++ b/Engine/Source/Editor/UserAssetTagsEditor/Private/UserAssetTagProviders/UserAssetTagProvider_LocalAssetTypeFavorites.cpp
@@ -27,10 +27,10 @@ TSet<FName> UUserAssetTagProvider_LocalAssetTypeFavorites::GetSuggestedUserAsset
 	if(UserAssetTagsEditor->GetSelectedAssets().Num() > 0)
 	{
 		// Due to the IsValid check we know there is only one class, so we fetch it from the first selected asset
-		TSharedPtr<FAssetData> SelectedAsset = UserAssetTagsEditor->GetSelectedAssets()[0];
-		if(ULocalFavoriteUserAssetTagsConfig::Get()->GetFavoriteUserAssetTags().Contains(FSoftClassPath(SelectedAsset->GetClass())))
+		const FAssetData& SelectedAsset = UserAssetTagsEditor->GetSelectedAssets()[0];
+		if(ULocalFavoriteUserAssetTagsConfig::Get()->GetFavoriteUserAssetTags().Contains(FSoftClassPath(SelectedAsset.GetClass())))
 		{
-			return ULocalFavoriteUserAssetTagsConfig::Get()->GetFavoriteUserAssetTags()[FSoftClassPath(SelectedAsset->GetClass())].FavoriteUserAssetTags;
+			return ULocalFavoriteUserAssetTagsConfig::Get()->GetFavoriteUserAssetTags()[FSoftClassPath(SelectedAsset.GetClass())].FavoriteUserAssetTags;
 		}
 	}
 
@@ -41,9 +41,9 @@ UUserAssetTagProvider::FResultWithUserFeedback UUserAssetTagProvider_LocalAssetT
 {
 	TSet<UClass*> Classes;
 
-	Algo::Transform(Context->UserAssetTagsEditor.Pin()->GetSelectedAssets(), Classes, [](const TSharedPtr<FAssetData>& AssetData)
+	Algo::Transform(Context->UserAssetTagsEditor.Pin()->GetSelectedAssets(), Classes, [](const FAssetData& AssetData)
 	{
-		return AssetData->GetClass();
+		return AssetData.GetClass();
 	});
 
 	FResultWithUserFeedback Result(Classes.Num() == 1);
@@ -62,7 +62,7 @@ TSharedPtr<SWidget> UUserAssetTagProvider_LocalAssetTypeFavorites::AddAdditional
 		{
 			if(Context->UserAssetTagsEditor.Pin()->GetSelectedAssets().Num() > 0)
 			{
-				UClass* Class = Context->UserAssetTagsEditor.Pin()->GetSelectedAssets()[0]->GetClass();
+				UClass* Class = Context->UserAssetTagsEditor.Pin()->GetSelectedAssets()[0].GetClass();
 				TSet<FName> FavoriteUserAssetTags = ULocalFavoriteUserAssetTagsConfig::Get()->GetFavoriteUserAssetTagsForClass(Class);
 				if(FavoriteUserAssetTags.Contains(UserAssetTag))
 				{
@@ -81,7 +81,7 @@ TSharedPtr<SWidget> UUserAssetTagProvider_LocalAssetTypeFavorites::AddAdditional
 		{
 			if(Context->UserAssetTagsEditor.Pin()->GetSelectedAssets().Num() > 0)
 			{
-				UClass* Class = Context->UserAssetTagsEditor.Pin()->GetSelectedAssets()[0]->GetClass();
+				UClass* Class = Context->UserAssetTagsEditor.Pin()->GetSelectedAssets()[0].GetClass();
 				ULocalFavoriteUserAssetTagsConfig::Get()->ToggleFavoriteUserAssetTag(Class, UserAssetTag);
 				Context->UserAssetTagsEditor.Pin()->RefreshDataAndMenus();
 			}
@@ -98,7 +98,7 @@ TSharedPtr<SWidget> UUserAssetTagProvider_LocalAssetTypeFavorites::AddAdditional
 				{
 					if(Context->UserAssetTagsEditor.Pin()->GetSelectedAssets().Num() > 0)
 					{
-						UClass* Class = Context->UserAssetTagsEditor.Pin()->GetSelectedAssets()[0]->GetClass();
+						UClass* Class = Context->UserAssetTagsEditor.Pin()->GetSelectedAssets()[0].GetClass();
 						TSet<FName> FavoriteUserAssetTags = ULocalFavoriteUserAssetTagsConfig::Get()->GetFavoriteUserAssetTagsForClass(Class);
 						if(FavoriteUserAssetTags.Contains(UserAssetTag))
 						{
diff --git a/Engine/Source/Editor/UserAssetTagsEditor/Private/UserAssetTagProviders/UserAssetTagProvider_Project.cpp b/Engine/Source/Editor/UserAssetTagsEditor/Private/UserAssetTagProviders/UserAssetTagProvider_Project.cpp
--- a/Engine/Source/Editor/UserAssetTagsEditor/Private/UserAssetTagProviders/UserAssetTagProvider_Project.cpp
+++ b/Engine/Source/Editor/UserAssetTagsEditor/Private/UserAssetTagProviders/UserAssetTagProvider_Project.cpp
@@ -18,8 +18,8 @@ FText UUserAssetTagProvider_Project::GetDisplayNameText(const UUserAssetTagEdito
 	TSharedPtr<SUserAssetTagsEditor> UserAssetTagsEditor = Context->UserAssetTagsEditor.Pin();
 	if(UserAssetTagsEditor->GetSelectedAssets().Num() > 0)
 	{
-		TSharedPtr<FAssetData> SelectedAsset = UserAssetTagsEditor->GetSelectedAssets()[0];
-		return FText::FormatOrdered(INVTEXT("{0} {1}"), SelectedAsset->GetClass()->GetDisplayNameText(), LOCTEXT("ProjectProviderDisplayName", "Project Tags"));
+		const FAssetData& SelectedAsset = UserAssetTagsEditor->GetSelectedAssets()[0];
+		return FText::FormatOrdered(INVTEXT("{0} {1}"), SelectedAsset.GetClass()->GetDisplayNameText(), LOCTEXT("ProjectProviderDisplayName", "Project Tags"));
 	}
 
 	return Super::GetDisplayNameText(Context);
@@ -31,8 +31,8 @@ TSet<FName> UUserAssetTagProvider_Project::GetSuggestedUserAssetTags(const UUser
 	if(UserAssetTagsEditor->GetSelectedAssets().Num() > 0)
 	{
 		// Due to the IsValid check we know there is only one class, so we fetch it from the first selected asset
-		TSharedPtr<FAssetData> SelectedAsset = UserAssetTagsEditor->GetSelectedAssets()[0];
-		return GetDefault<UProjectUserAssetTagSettings>()->GetProjectUserAssetTagsForClass(SelectedAsset->GetClass());
+		const FAssetData& SelectedAsset = UserAssetTagsEditor->GetSelectedAssets()[0];
+		return GetDefault<UProjectUserAssetTagSettings>()->GetProjectUserAssetTagsForClass(SelectedAsset.GetClass());
 	}
 
 	return {};
@@ -42,9 +42,9 @@ UUserAssetTagProvider::FResultWithUserFeedback UUserAssetTagProvider_Project::Is
 {
 	TSet<UClass*> Classes;
 
-	Algo::Transform(Context->UserAssetTagsEditor.Pin()->GetSelectedAssets(), Classes, [](const TSharedPtr<FAssetData>& AssetData)
+	Algo::Transform(Context->UserAssetTagsEditor.Pin()->GetSelectedAssets(), Classes, [](const FAssetData& AssetData)
 	{
-		return AssetData->GetClass();
+		return AssetData.GetClass();
 	});
 
 	FResultWithUserFeedback Result(Classes.Num() == 1);
diff --git a/Engine/Source/Editor/UserAssetTagsEditor/Private/Widgets/SUserAssetTagsEditor.cpp b/Engine/Source/Editor/UserAssetTagsEditor/Private/Widgets/SUserAssetTagsEditor.cpp
--- a/Engine/Source/Editor/UserAssetTagsEditor/Private/Widgets/SUserAssetTagsEditor.cpp
+++ b/Engine/Source/Editor/UserAssetTagsEditor/Private/Widgets/SUserAssetTagsEditor.cpp
@@ -152,12 +152,10 @@ TSharedPtr<SWidget> SUserAssetTagsEditor::CreateSelectedAssetList()
 			}
 			
 			FToolMenuSection& OwnedTagsSection = DynamicMenu->FindOrAddSection(FName(TEXT("SelectedAssets")), LOCTEXT("SelectedAssetsLabel", "Selected Assets"));
-
-			const TArray<FAssetData> SelectedAssets = TransformAssetData(ContextObject->UserAssetTagsEditor.Pin()->GetSelectedAssets());
-			
-			FMultipleAssetsTagInfo TagInfo = GatherInfoAboutTags(SelectedAssets);
+				
+			FMultipleAssetsTagInfo TagInfo = GatherInfoAboutTags(ContextObject->UserAssetTagsEditor.Pin()->GetSelectedAssets());
 			
-			for(const FAssetData& SelectedAsset : SelectedAssets)
+			for(const FAssetData& SelectedAsset : ContextObject->UserAssetTagsEditor.Pin()->GetSelectedAssets())
 			{				
 				FToolMenuEntry Entry = FToolMenuEntry::InitMenuEntry(SelectedAsset.AssetName, FToolUIActionChoice(), ContextObject->UserAssetTagsEditor.Pin()->GenerateRowContent_SelectedAsset(SelectedAsset, ContextObject));
 				Entry.UserInterfaceActionType = EUserInterfaceActionType::None;
@@ -209,8 +207,6 @@ TSharedPtr<SWidget> SUserAssetTagsEditor::CreateSuggestedTagList()
 				}
 			}
 
-			const TArray<FAssetData> ApplicableAssets = TransformAssetData(ContextObject->UserAssetTagsEditor.Pin()->GetSelectedAssets());
-
 			for(const auto& ProviderTagsPair : ProviderTagMap)
 			{
 				EUserAssetTagProviderMenuType MenuType = GetViewOptions_ProviderClassMenuType(ProviderTagsPair.Key->GetClass());
@@ -293,15 +289,13 @@ TSharedPtr<SWidget> SUserAssetTagsEditor::CreateOwnedTagsList()
 			}
 			
 			FToolMenuSection& OwnedTagsSection = DynamicMenu->FindOrAddSection(FName(TEXT("OwnedTags")), LOCTEXT("OwnedTagsLabel", "Owned Tags"));
-
-			const TArray<FAssetData> ApplicableAssets = TransformAssetData(ContextObject->UserAssetTagsEditor.Pin()->GetSelectedAssets());
 			
-			FMultipleAssetsTagInfo TagInfo = GatherInfoAboutTags(ApplicableAssets);
+			FMultipleAssetsTagInfo TagInfo = GatherInfoAboutTags(ContextObject->UserAssetTagsEditor.Pin()->GetSelectedAssets());
 			
 			for(const auto& TagPairInfo : TagInfo.TagInfos)
 			{
 				FToolMenuEntry Entry = FToolMenuEntry::InitMenuEntry(TagPairInfo.Key, CreateToggleTagCheckboxAction_OwnedTag(TagPairInfo.Key), ContextObject->UserAssetTagsEditor.Pin()->GenerateRowContent_OwnedTag(TagPairInfo.Key, ContextObject));
-				ECheckBoxState CheckBoxState = DetermineTagUIStatus(TagPairInfo.Key, TransformAssetData(ContextObject->UserAssetTagsEditor.Pin()->GetSelectedAssets()));
+				ECheckBoxState CheckBoxState = DetermineTagUIStatus(TagPairInfo.Key, ContextObject->UserAssetTagsEditor.Pin()->GetSelectedAssets());
 				// We only want the checkbox to show up when there the tag status is undetermined (aka some selected assets have the tag, some don't).
 				// The main interface to delete tags is the delete button, but to offer consistent UI with the suggested tags we also make use of a checkbox
 				// This way you can either 'fix' the tag status by clicking on the checkbox (tagging all assets), or click on the delete button to delete the tags specifically
@@ -410,11 +404,7 @@ TArray<FAssetData> SUserAssetTagsEditor::GetCurrentContentBrowserAssetSelection(
 void SUserAssetTagsEditor::RefreshDataAndMenus()
 {
 	SelectedAssets.Empty();
-	Algo::Transform(GetCurrentContentBrowserAssetSelection(), SelectedAssets, [](FAssetData InAssetData)
-	{
-		return MakeShared<FAssetData>(InAssetData);
-	});
-
+	SelectedAssets = GetCurrentContentBrowserAssetSelection();
 	RefreshMenus();
 }
 
@@ -449,10 +439,8 @@ void SUserAssetTagsEditor::RefreshProviderExtensionToolbarMenu()
 void SUserAssetTagsEditor::OnCommitNewTag(const FText& Text, ETextCommit::Type CommitType)
 {
 	if(CommitType == ETextCommit::Type::OnEnter && Text.IsEmptyOrWhitespace() == false)
-	{
-		TArray<FAssetData> ApplicableAssets = TransformAssetData(SelectedAssets);
-		
-		UE::UserAssetTags::Menus::AddTagsToAssets(FName(Text.ToString()), ApplicableAssets);
+	{		
+		UE::UserAssetTags::Menus::AddTagsToAssets(FName(Text.ToString()), SelectedAssets);
 
 		RefreshDataAndMenus();
 	}
@@ -464,9 +452,8 @@ FReply SUserAssetTagsEditor::OnAddTagButtonClicked()
 	{
 		return FReply::Handled();
 	}
-
-	TArray<FAssetData> ApplicableAssets = TransformAssetData(SelectedAssets);
-	UE::UserAssetTags::Menus::AddTagsToAssets(FName(AddTagTextBox->GetText().ToString()), ApplicableAssets);
+	
+	UE::UserAssetTags::Menus::AddTagsToAssets(FName(AddTagTextBox->GetText().ToString()), SelectedAssets);
 	RefreshDataAndMenus();
 	return FReply::Handled();
 }
@@ -821,12 +808,12 @@ FToolUIAction SUserAssetTagsEditor::CreateToggleTagCheckboxAction_SuggestedTag(c
 	Action.GetActionCheckState = FToolMenuGetActionCheckState::CreateLambda([InUserAssetTag](const FToolMenuContext& Context)
 	{
 		UUserAssetTagEditorContext* ContextObject = Context.FindContext<UUserAssetTagEditorContext>();
-		return DetermineTagUIStatus(InUserAssetTag, TransformAssetData(ContextObject->UserAssetTagsEditor.Pin()->GetSelectedAssets()));
+		return DetermineTagUIStatus(InUserAssetTag, ContextObject->UserAssetTagsEditor.Pin()->GetSelectedAssets());
 	});
 	Action.ExecuteAction = FToolMenuExecuteAction::CreateLambda([InUserAssetTag](const FToolMenuContext& Context)
 	{
 		UUserAssetTagEditorContext* ContextObject = Context.FindContext<UUserAssetTagEditorContext>();
-		TArray<FAssetData> Assets = TransformAssetData(ContextObject->UserAssetTagsEditor.Pin()->GetSelectedAssets());
+		const TArray<FAssetData>& Assets = ContextObject->UserAssetTagsEditor.Pin()->GetSelectedAssets();
 					
 		ECheckBoxState CurrentState = DetermineTagUIStatus(InUserAssetTag, Assets);
 		if(CurrentState == ECheckBoxState::Checked)
@@ -856,12 +843,12 @@ FToolUIAction SUserAssetTagsEditor::CreateToggleTagCheckboxAction_OwnedTag(const
 	Action.GetActionCheckState = FToolMenuGetActionCheckState::CreateLambda([InUserAssetTag](const FToolMenuContext& Context)
 	{
 		UUserAssetTagEditorContext* ContextObject = Context.FindContext<UUserAssetTagEditorContext>();
-		return DetermineTagUIStatus(InUserAssetTag, TransformAssetData(ContextObject->UserAssetTagsEditor.Pin()->GetSelectedAssets()));
+		return DetermineTagUIStatus(InUserAssetTag, ContextObject->UserAssetTagsEditor.Pin()->GetSelectedAssets());
 	});
 	Action.ExecuteAction = FToolMenuExecuteAction::CreateLambda([InUserAssetTag](const FToolMenuContext& Context)
 	{
 		UUserAssetTagEditorContext* ContextObject = Context.FindContext<UUserAssetTagEditorContext>();
-		TArray<FAssetData> Assets = TransformAssetData(ContextObject->UserAssetTagsEditor.Pin()->GetSelectedAssets());
+		const TArray<FAssetData>& Assets = ContextObject->UserAssetTagsEditor.Pin()->GetSelectedAssets();
 					
 		ECheckBoxState CurrentState = DetermineTagUIStatus(InUserAssetTag, Assets);
 		if(CurrentState == ECheckBoxState::Checked)
@@ -911,7 +898,7 @@ FText SUserAssetTagsEditor::GetDocumentationButtonTooltipText()
 FReply SUserAssetTagsEditor::OnDeleteTagClicked(FName InUserAssetTag, const UUserAssetTagEditorContext* Context)
 {
 	using namespace UE::UserAssetTags::Menus;
-	TArray<FAssetData> Assets = TransformAssetData(Context->UserAssetTagsEditor.Pin()->GetSelectedAssets());
+	const TArray<FAssetData>& Assets = Context->UserAssetTagsEditor.Pin()->GetSelectedAssets();
 	RemoveTagsFromAssets(InUserAssetTag, Assets);
 	RefreshMenus();
 	return FReply::Handled();
@@ -922,17 +909,6 @@ void SUserAssetTagsEditor::HandleAssetSelectionChanged(const TArray<FAssetData>&
 	RefreshDataAndMenus();
 }
 
-TArray<FAssetData> SUserAssetTagsEditor::TransformAssetData(const TArray<TSharedPtr<FAssetData>>& InAssetData)
-{
-	TArray<FAssetData> Result;
-	Algo::Transform(InAssetData, Result, [](TSharedPtr<FAssetData> InAssetData)
-	{
-		return *InAssetData.Get();
-	});
-
-	return Result;
-}
-
 #undef LOCTEXT_NAMESPACE
 
 
diff --git a/Engine/Source/Editor/UserAssetTagsEditor/Public/Widgets/SUserAssetTagsEditor.h b/Engine/Source/Editor/UserAssetTagsEditor/Public/Widgets/SUserAssetTagsEditor.h
--- a/Engine/Source/Editor/UserAssetTagsEditor/Public/Widgets/SUserAssetTagsEditor.h
+++ b/Engine/Source/Editor/UserAssetTagsEditor/Public/Widgets/SUserAssetTagsEditor.h
@@ -35,7 +35,7 @@ public:
 	UE_API void Construct(const FArguments& InArgs);
 	UE_API virtual ~SUserAssetTagsEditor() override;
 	
-	const TArray<TSharedPtr<FAssetData>>& GetSelectedAssets() const { return SelectedAssets; }
+	const TArray<FAssetData>& GetSelectedAssets() const { return SelectedAssets; }
 
 	UE_API void RefreshDataAndMenus();
 	static UE_API void RefreshMenus();
@@ -88,14 +88,13 @@ private:
 	static UE_API FReply OnDeleteTagClicked(FName InUserAssetTag, const UUserAssetTagEditorContext* Context);
 
 	UE_API void HandleAssetSelectionChanged(const TArray<FAssetData>& AssetData, bool bIsPrimaryBrowser);
-	
-	static UE_API TArray<FAssetData> TransformAssetData(const TArray<TSharedPtr<FAssetData>>& InAssetData);
+
 private:
 	TSharedPtr<SWidget> SelectedAssetList;
 	TSharedPtr<SWidget> TagSuggestionList;
 	TSharedPtr<SWidget> OwnedTagsList;
 
-	TArray<TSharedPtr<FAssetData>> SelectedAssets;
+	TArray<FAssetData> SelectedAssets;
 	
 	TSharedPtr<class SEditableTextBox> AddTagTextBox;
 	TSharedPtr<SWidget> ToolbarWidget;
