TaggedAssetBrowserConfig - Add HasInstance call to prevent creating instance on destruction

#rnx
#Synced-CL 49823132

diff --git a/Engine/Source/Editor/UserAssetTagsEditor/Private/Config/TaggedAssetBrowserConfig.cpp b/Engine/Source/Editor/UserAssetTagsEditor/Private/Config/TaggedAssetBrowserConfig.cpp
--- a/Engine/Source/Editor/UserAssetTagsEditor/Private/Config/TaggedAssetBrowserConfig.cpp
+++ b/Engine/Source/Editor/UserAssetTagsEditor/Private/Config/TaggedAssetBrowserConfig.cpp
@@ -2,6 +2,8 @@
 
 #include "Config/TaggedAssetBrowserConfig.h"
 
+#include "UObject/StrongObjectPtrTemplates.h"
+
 TStrongObjectPtr<UTaggedAssetBrowserConfig> UTaggedAssetBrowserConfig::Instance = nullptr;
 
 UTaggedAssetBrowserConfig* UTaggedAssetBrowserConfig::Get()
@@ -24,6 +26,11 @@ void UTaggedAssetBrowserConfig::Shutdown()
 	}
 }
 
+bool UTaggedAssetBrowserConfig::HasInstance()
+{
+	return Instance.IsValid();
+}
+
 void UTaggedAssetBrowserConfig::PostEditChangeProperty(FPropertyChangedEvent& PropertyChangedEvent)
 {
 	Super::PostEditChangeProperty(PropertyChangedEvent);
diff --git a/Engine/Source/Editor/UserAssetTagsEditor/Private/Widgets/STaggedAssetBrowser.cpp b/Engine/Source/Editor/UserAssetTagsEditor/Private/Widgets/STaggedAssetBrowser.cpp
--- a/Engine/Source/Editor/UserAssetTagsEditor/Private/Widgets/STaggedAssetBrowser.cpp
+++ b/Engine/Source/Editor/UserAssetTagsEditor/Private/Widgets/STaggedAssetBrowser.cpp
@@ -115,7 +115,10 @@ STaggedAssetBrowser::STaggedAssetBrowser()
 STaggedAssetBrowser::~STaggedAssetBrowser()
 {
 	SaveSettings();
-	UTaggedAssetBrowserConfig::Get()->OnPropertyChanged().RemoveAll(this);
+	if (UTaggedAssetBrowserConfig::HasInstance())
+	{
+		UTaggedAssetBrowserConfig::Get()->OnPropertyChanged().RemoveAll(this);
+	}
 }
 
 TArray<UClass*> STaggedAssetBrowser::GetDisplayedAssetTypes() const
diff --git a/Engine/Source/Editor/UserAssetTagsEditor/Public/Config/TaggedAssetBrowserConfig.h b/Engine/Source/Editor/UserAssetTagsEditor/Public/Config/TaggedAssetBrowserConfig.h
--- a/Engine/Source/Editor/UserAssetTagsEditor/Public/Config/TaggedAssetBrowserConfig.h
+++ b/Engine/Source/Editor/UserAssetTagsEditor/Public/Config/TaggedAssetBrowserConfig.h
@@ -2,7 +2,6 @@
 
 #pragma once
 
-#include "CoreMinimal.h"
 #include "EditorConfigBase.h"
 #include "TaggedAssetBrowserConfig.generated.h"
 
@@ -26,6 +25,7 @@ public:
 		
 	static UE_API UTaggedAssetBrowserConfig* Get();
 	static UE_API void Shutdown();
+	static UE_API bool HasInstance();
 	
 	FOnPropertyChanged& OnPropertyChanged() { return OnPropertyChangedDelegate; }
 	
