Fixed rightclick crash on tagged asset browser sections

#Synced-CL 51048610

diff --git a/Engine/Source/Editor/UserAssetTagsEditor/Private/Widgets/STaggedAssetBrowserSections.cpp b/Engine/Source/Editor/UserAssetTagsEditor/Private/Widgets/STaggedAssetBrowserSections.cpp
--- a/Engine/Source/Editor/UserAssetTagsEditor/Private/Widgets/STaggedAssetBrowserSections.cpp
+++ b/Engine/Source/Editor/UserAssetTagsEditor/Private/Widgets/STaggedAssetBrowserSections.cpp
@@ -98,7 +98,12 @@ ECheckBoxState STaggedAssetBrowserSection::GetCheckState() const
 
 TSharedRef<SWidget> STaggedAssetBrowserSection::OnGetMenuContent()
 {
-	return OnGetMenuContentDelegate.Execute();
+	if (OnGetMenuContentDelegate.IsBound())
+	{
+		return OnGetMenuContentDelegate.Execute();
+	}
+	
+	return SNullWidget::NullWidget;
 }
 
 FSlateColor STaggedAssetBrowserSection::GetIconForegroundColor() const
