Add UAFSharedAssets plugin for content we want to provide that references UAF assets defined in separate, independent, plugins.
- Use new plugin to add ST / Pose search to UAF Browser asset config.
- Update Tagged Asset Browser to support merging sections of the same name across extension configs.
- Add `AppendAfterFilterName` to specify where to merge (Hidden if not extension)
- Add `bMergeExtensionsPerSection` to allow for merging of same named sections (Enabled by default).

#tests editor
#virtualized
#Synced-CL 52638195

#Synced-CL 52638195

diff --git a/Engine/Plugins/Experimental/UAF/UAFSharedAssets/UAFSharedAssets.uplugin b/Engine/Plugins/Experimental/UAF/UAFSharedAssets/UAFSharedAssets.uplugin
new file mode 100644
--- /dev/null
+++ b/Engine/Plugins/Experimental/UAF/UAFSharedAssets/UAFSharedAssets.uplugin
@@ -0,0 +1,57 @@
+{
+	"FileVersion": 1,
+	"Version": 1,
+	"VersionName": "1.0",
+	"FriendlyName": "UAF Shared Assets",
+	"Description": "UAF Default Assets that interact with multiple plugins",
+	"Category": "Other",
+	"CreatedBy": "Epic Games, Inc.",
+	"CreatedByURL": "https://epicgames.com",
+	"DocsURL": "",
+	"MarketplaceURL": "",
+	"SupportURL": "",
+	"CanContainContent": true,
+	"IsBetaVersion": false,
+	"IsExperimentalVersion": true,
+	"Installed": false,
+	"Modules": [
+	],
+	"Plugins": [
+		{
+			"Name": "Workspace",
+			"Enabled": true
+		},
+		{
+			"Name": "Chooser",
+			"Enabled": true
+		},
+		{
+			"Name": "EvaluationNotifies",
+			"Enabled": true
+		},
+		{
+			"Name": "PoseSearch",
+			"Enabled": true
+		},
+		{
+			"Name": "RigVM",
+			"Enabled": true		
+		},
+		{
+			"Name": "UAF",
+			"Enabled": true
+		},
+		{
+			"Name": "UAFAnimNode",
+			"Enabled": true
+		},
+		{
+			"Name": "UAFAnimGraph",
+			"Enabled": true		
+		},		
+		{
+			"Name": "UAFStateTree",
+			"Enabled": true		
+		}
+	]
+}
diff --git a/Engine/Plugins/Experimental/UAF/UAFStateTree/UAFStateTree.uplugin b/Engine/Plugins/Experimental/UAF/UAFStateTree/UAFStateTree.uplugin
--- a/Engine/Plugins/Experimental/UAF/UAFStateTree/UAFStateTree.uplugin
+++ b/Engine/Plugins/Experimental/UAF/UAFStateTree/UAFStateTree.uplugin
@@ -11,7 +11,7 @@
 	"MarketplaceURL" : "",
 	"SupportURL" : "",
 	"EnabledByDefault" : false,
-	"CanContainContent" : false,
+	"CanContainContent" : true,
 	"IsBetaVersion": false,
 	"IsExperimentalVersion": true,
 	"Installed" : false,
diff --git a/Engine/Source/Editor/UserAssetTagsEditor/Private/TaggedAssetBrowserMenuFilters.cpp b/Engine/Source/Editor/UserAssetTagsEditor/Private/TaggedAssetBrowserMenuFilters.cpp
--- a/Engine/Source/Editor/UserAssetTagsEditor/Private/TaggedAssetBrowserMenuFilters.cpp
+++ b/Engine/Source/Editor/UserAssetTagsEditor/Private/TaggedAssetBrowserMenuFilters.cpp
@@ -2,6 +2,8 @@
 
 #include "TaggedAssetBrowserMenuFilters.h"
 
+#include "Assets/TaggedAssetBrowserConfiguration.h"
+
 #define LOCTEXT_NAMESPACE "UserAssetTagMenuFilters"
 
 FName UTaggedAssetBrowserFilterBase::GetIdentifier() const
@@ -36,6 +38,16 @@ bool UTaggedAssetBrowserFilterBase::ShouldFilterAsset(const FAssetData& InAssetD
 	return ShouldFilterAssetInternal(InAssetData);
 }
 
+bool UTaggedAssetBrowserFilterBase::IsExtensionFilter() const
+{
+	if (UTaggedAssetBrowserConfiguration* OwningConfig = GetTypedOuter<UTaggedAssetBrowserConfiguration>())
+	{
+		return OwningConfig->bIsExtension;
+	}
+
+	return false;
+}
+
 void UTaggedAssetBrowserFilterBase::Initialize(const FTaggedAssetBrowserContext& InContext)
 {
 	ActiveContext = InContext;
diff --git a/Engine/Source/Editor/UserAssetTagsEditor/Private/Widgets/STaggedAssetBrowser.cpp b/Engine/Source/Editor/UserAssetTagsEditor/Private/Widgets/STaggedAssetBrowser.cpp
--- a/Engine/Source/Editor/UserAssetTagsEditor/Private/Widgets/STaggedAssetBrowser.cpp
+++ b/Engine/Source/Editor/UserAssetTagsEditor/Private/Widgets/STaggedAssetBrowser.cpp
@@ -195,28 +195,117 @@ void STaggedAssetBrowser::ApplyFilterExtensions()
 {
 	if(TaggedBrowserProfileName.IsSet() && TaggedBrowserProfileName->IsValid())
 	{
-		TMap<FName, UHierarchySection*> NewSectionMap;
+		TMap<FName, const UHierarchySection*> MergeAllowedSectionMap;
+		TMap<FName, const UHierarchySection*> MergeDisabledSectionMap;
+
+		// Build a map of sections already present in the base config so we can merge into them
+		TStrongObjectPtr<const UTaggedAssetBrowserConfiguration> Config = Configuration.Pin();
+		bool bMergeExtensions = Config && Config->StandaloneData.bMergeExtensionsPerSection;
+		if (bMergeExtensions)
+		{
+			for (const UHierarchySection* ExistingSection : FilterHierarchyRoot->GetSectionData())
+			{
+				const UTaggedAssetBrowserSection* TaggedAssetBrowserSection = Cast<UTaggedAssetBrowserSection>(ExistingSection);
+				if (TaggedAssetBrowserSection && TaggedAssetBrowserSection->bAllowSectionMerge)
+				{
+					MergeAllowedSectionMap.Add(ExistingSection->GetSectionName(), ExistingSection);
+				}
+				else
+				{
+					MergeDisabledSectionMap.Add(ExistingSection->GetSectionName(), ExistingSection);
+				}
+			}
+		}
+
 		for(const UTaggedAssetBrowserConfiguration* ExtensionConfigurationAsset : GetExtensionAssets())
 		{
 			for(const UHierarchySection* Section : ExtensionConfigurationAsset->FilterRoot->GetSectionData())
 			{
-				UHierarchySection* CopiedSection = Cast<UHierarchySection>(StaticDuplicateObject(Section, FilterHierarchyRoot.Get()));
-				FilterHierarchyRoot->GetSectionDataMutable().Add(CopiedSection);
-				NewSectionMap.Add(CopiedSection->GetSectionName(), CopiedSection);
+				const FName SectionName = Section->GetSectionName();
+				if(!MergeAllowedSectionMap.Contains(SectionName) || !bMergeExtensions)
+				{
+					// No matching section in the base config — add a new one
+					UHierarchySection* CopiedSection = Cast<UHierarchySection>(StaticDuplicateObject(Section, FilterHierarchyRoot.Get()));
+					FilterHierarchyRoot->GetSectionDataMutable().Add(CopiedSection);
+
+					const UTaggedAssetBrowserSection* TaggedAssetBrowserSection = Cast<UTaggedAssetBrowserSection>(CopiedSection);
+					if (TaggedAssetBrowserSection && TaggedAssetBrowserSection->bAllowSectionMerge)
+					{
+						MergeAllowedSectionMap.Add(CopiedSection->GetSectionName(), CopiedSection);
+					}
+					else
+					{
+						MergeDisabledSectionMap.Add(CopiedSection->GetSectionName(), CopiedSection);
+					}
+				}
+				else
+				{
+					// Do nothing. If the section already exists, we will add the elements to it below
+				}
 			}
-			
+
 			// We copy over the children and append them to the root
 			for(const UHierarchyElement* Element : ExtensionConfigurationAsset->FilterRoot->GetChildren())
 			{
-				UHierarchyElement* CopiedItem = FilterHierarchyRoot->CopyAndAddItemAsChild(*Element);
+				// Manually insert into appropriate location 
+				UHierarchyElement* NewChild = Cast<UHierarchyElement>(StaticDuplicateObject(Element, FilterHierarchyRoot.Get()));
+				if (NewChild->GetPersistentIdentity() != Element->GetPersistentIdentity())
+				{
+					check(false);
+				}
+
+				bool bFilterInserted = false;
+				if (const UTaggedAssetBrowserFilterBase* Filter = Cast<UTaggedAssetBrowserFilterBase>(Element))
+				{
+					auto IsInsertTargetFilter = [FilterName = Filter->ExtensionData.InsertFilterName](const UHierarchyElement* ExistingElement)
+					{
+						if (const UTaggedAssetBrowserFilterBase* ExistingFilter = Cast<UTaggedAssetBrowserFilterBase>(ExistingElement))
+						{
+							return ExistingFilter->ToString() == FilterName.ToString();
+						}
+
+						return false;
+					};
+
+					if (bMergeExtensions 
+						&& Filter->ExtensionData.MergeBehavior == ETaggedAssetBrowserFilterMergeBehavior::Insert
+						&& !Filter->ExtensionData.InsertFilterName.IsNone())
+					{
+						int32 InsertIndex = FilterHierarchyRoot->GetChildrenMutable().IndexOfByPredicate(IsInsertTargetFilter);
+
+						if (InsertIndex != INDEX_NONE)
+						{
+							if (Filter->ExtensionData.InsertBehavior == ETaggedAssetBrowserFilterInsertBehavior::After)
+							{
+								InsertIndex++;
+							}
+						}
+						else
+						{
+							InsertIndex = FilterHierarchyRoot->GetChildrenMutable().Num();
+						}
+
+						FilterHierarchyRoot->GetChildrenMutable().Insert(NewChild, InsertIndex);
+						bFilterInserted = true;
+					}
+				}
+
+				if (!bFilterInserted)
+				{
+					FilterHierarchyRoot->GetChildrenMutable().Add(NewChild);
+				}
 
 				// We also need to fix up section associations
-				FDataHierarchyElementMetaData_SectionAssociation* SectionAssociation = CopiedItem->FindMetaDataOfType<FDataHierarchyElementMetaData_SectionAssociation>();
+				FDataHierarchyElementMetaData_SectionAssociation* SectionAssociation = NewChild->FindMetaDataOfType<FDataHierarchyElementMetaData_SectionAssociation>();
 				if(SectionAssociation && SectionAssociation->Section.IsValid())
 				{
-					if(NewSectionMap.Contains(SectionAssociation->Section->GetSectionName()))
+					if(MergeAllowedSectionMap.Contains(SectionAssociation->Section->GetSectionName()))
+					{
+						SectionAssociation->Section = MergeAllowedSectionMap[SectionAssociation->Section->GetSectionName()];
+					}
+					else if (MergeDisabledSectionMap.Contains(SectionAssociation->Section->GetSectionName()))
 					{
-						SectionAssociation->Section = NewSectionMap[SectionAssociation->Section->GetSectionName()];
+						SectionAssociation->Section = MergeDisabledSectionMap[SectionAssociation->Section->GetSectionName()];
 					}
 					else
 					{
diff --git a/Engine/Source/Editor/UserAssetTagsEditor/Public/Assets/TaggedAssetBrowserConfiguration.h b/Engine/Source/Editor/UserAssetTagsEditor/Public/Assets/TaggedAssetBrowserConfiguration.h
--- a/Engine/Source/Editor/UserAssetTagsEditor/Public/Assets/TaggedAssetBrowserConfiguration.h
+++ b/Engine/Source/Editor/UserAssetTagsEditor/Public/Assets/TaggedAssetBrowserConfiguration.h
@@ -25,6 +25,10 @@ struct FTaggedAssetBrowserConfigurationData_Standalone : public FTaggedAssetBrow
 	
 	UPROPERTY(EditAnywhere, Category="Configuration")
 	bool bDisplayAllSection = true;
+
+	/** True implies we should merge extensions by section */
+	UPROPERTY(EditAnywhere, Category = "Configuration")
+	bool bMergeExtensionsPerSection = true;
 	
 	/** The list of available hierarchy elements to add in the filter hierarchy for standalone configuration assets. */
 	UPROPERTY(EditAnywhere, Category="Configuration", meta=(ShowDisplayNames))
diff --git a/Engine/Source/Editor/UserAssetTagsEditor/Public/TaggedAssetBrowserFilters/TaggedAssetBrowser_CommonFilters.h b/Engine/Source/Editor/UserAssetTagsEditor/Public/TaggedAssetBrowserFilters/TaggedAssetBrowser_CommonFilters.h
--- a/Engine/Source/Editor/UserAssetTagsEditor/Public/TaggedAssetBrowserFilters/TaggedAssetBrowser_CommonFilters.h
+++ b/Engine/Source/Editor/UserAssetTagsEditor/Public/TaggedAssetBrowserFilters/TaggedAssetBrowser_CommonFilters.h
@@ -52,6 +52,10 @@ public:
 	UPROPERTY(EditAnywhere, Category = "Section", meta=(ShowOnlyInnerProperties))
 	FTaggedAssetBrowserSectionIconData IconData;
 
+	/** True implies this section should support merging  */
+	UPROPERTY(EditAnywhere, Category = "Section")
+	bool bAllowSectionMerge = true;
+
 	/** A list of filters that can be specified optionally. This way, a section can display its associated elements, OR additionally filter assets by itself. */
 	UPROPERTY(EditAnywhere, Category = "Section", Instanced)
 	TArray<TObjectPtr<UTaggedAssetBrowserFilterBase>> Filters;
diff --git a/Engine/Source/Editor/UserAssetTagsEditor/Public/TaggedAssetBrowserMenuFilters.h b/Engine/Source/Editor/UserAssetTagsEditor/Public/TaggedAssetBrowserMenuFilters.h
--- a/Engine/Source/Editor/UserAssetTagsEditor/Public/TaggedAssetBrowserMenuFilters.h
+++ b/Engine/Source/Editor/UserAssetTagsEditor/Public/TaggedAssetBrowserMenuFilters.h
@@ -19,6 +19,43 @@ struct FTaggedAssetBrowserContext
 	FOnGetSelectedFilters OnGetSelectedFilters;
 };
 
+/** How to resolve extension filter merging */
+UENUM()
+enum class ETaggedAssetBrowserFilterMergeBehavior
+{
+	Append, // Add Filter at end
+	Insert, // Insert before or after a filter by name
+};
+
+/** How to insert a filter */
+UENUM()
+enum class ETaggedAssetBrowserFilterInsertBehavior
+{
+	Before,
+	After
+};
+
+/** Data relevant for extension filters. */
+USTRUCT()
+struct FTaggedAssetBrowserFilterData_Extension 
+{
+	GENERATED_BODY()
+
+public:
+
+	/** How this filter should be merged */
+	UPROPERTY(EditAnywhere, Category = "Filter")
+	ETaggedAssetBrowserFilterMergeBehavior MergeBehavior = ETaggedAssetBrowserFilterMergeBehavior::Append;
+
+	/** How this filter should be inserted */
+	UPROPERTY(EditAnywhere, Category = "Filter", meta = (EditCondition = "MergeBehavior==ETaggedAssetBrowserFilterMergeBehavior::Insert", EditConditionHides))
+	ETaggedAssetBrowserFilterInsertBehavior InsertBehavior = ETaggedAssetBrowserFilterInsertBehavior::After;
+
+	/** Name for filter this filter should be inserted relative to, assuming this is from / in an extension. */
+	UPROPERTY(EditAnywhere, Category = "Filter", meta = (EditCondition = "MergeBehavior==ETaggedAssetBrowserFilterMergeBehavior::Insert", EditConditionHides))
+	FName InsertFilterName;
+};
+
 /** A filter is the base element of the filter hierarchy built in the TaggedAssetBrowserConfiguration asset.
  *  It filters using two methods:
  *  - ModifyARFilter, which modifies the filter used to query the Asset Registry (override the Internal function for your implementation)
@@ -30,7 +67,18 @@ UCLASS(MinimalAPI, Abstract, EditInlineNew)
 class UTaggedAssetBrowserFilterBase : public UHierarchyItem
 {
 	GENERATED_BODY()
-	
+
+public:
+
+	/** Data for this filter that is only needed when we are in an extension. */
+	UPROPERTY(EditAnywhere, Category = "Filter", meta = (EditCondition = "IsExtensionFilter", EditConditionHides))
+	FTaggedAssetBrowserFilterData_Extension ExtensionData;
+
+public:
+	/** Check outer for if we are a filter in an extension config */
+	UFUNCTION()
+	bool IsExtensionFilter() const;
+
 public:	
 	/** The filter has a chance to initialize itself with contextual data to actually work. */
 	UE_API void Initialize(const FTaggedAssetBrowserContext& InContext);
